Fix crash auto-keying FColor values in sequencer

diff --git a/Engine/Source/Developer/MovieSceneTools/Private/PropertyTrackEditor.cpp b/Engine/Source/Developer/MovieSceneTools/Private/PropertyTrackEditor.cpp
--- a/Engine/Source/Developer/MovieSceneTools/Private/PropertyTrackEditor.cpp
+++ b/Engine/Source/Developer/MovieSceneTools/Private/PropertyTrackEditor.cpp
@@ -422,7 +422,7 @@ void FPropertyTrackEditor::OnAnimatedColorPropertyChanged( const TArray<UObject*
 		PropertyValue.GetChildHandle("A")
 	};
 
-	int32 FColorChannels[4];
+	uint8 FColorChannels[4];
 	float FLinearColorChannels[4];
 	for (int32 i = 0; i < 4; ++i)
 	{
@@ -442,6 +442,11 @@ void FPropertyTrackEditor::OnAnimatedColorPropertyChanged( const TArray<UObject*
 		ColorValue = FLinearColor(FLinearColorChannels[0], FLinearColorChannels[1], FLinearColorChannels[2], FLinearColorChannels[3]);
 	}
 	
+	if( StructProp->HasMetaData("HideAlphaChannel") )
+	{
+		ColorValue.A = 1;
+	}
+
 	AnimatablePropertyChanged(UMovieSceneColorTrack::StaticClass(), bRequireAutoKey,
 		FOnKeyProperty::CreateRaw(this, &FPropertyTrackEditor::OnKeyProperty<FLinearColor, UMovieSceneColorTrack>, PropertyName, InObjectsThatChanged, (const FLinearColor*)&ColorValue, bRequireAutoKey));
 }
