Fix non-unity

diff --git a/Engine/Source/Developer/MovieSceneTools/Private/PropertyTrackEditor.h b/Engine/Source/Developer/MovieSceneTools/Private/PropertyTrackEditor.h
--- a/Engine/Source/Developer/MovieSceneTools/Private/PropertyTrackEditor.h
+++ b/Engine/Source/Developer/MovieSceneTools/Private/PropertyTrackEditor.h
@@ -4,6 +4,8 @@
 
 class IPropertyHandle;
 
+struct FKeyPropertyParams;
+
 /**
  * Tools for animatable property types such as floats ands vectors
  */
@@ -39,7 +41,7 @@ private:
 	 * @param PropertyValue			Handle to the property value which changed
 	 */
 	template <typename Type, typename TrackType>
-	void OnAnimatedPropertyChanged( const struct FKeyPropertyParams& KeyPropertyParams );
+	void OnAnimatedPropertyChanged( const FKeyPropertyParams& KeyPropertyParams );
 
 	/**
 	 * Called by the details panel when an animatable vector property changes
@@ -47,7 +49,7 @@ private:
 	 * @param InObjectsThatChanged	List of objects that changed
 	 * @param PropertyValue			Handle to the property value which changed
 	 */
-	void OnAnimatedVectorPropertyChanged( const struct FKeyPropertyParams& KeyPropertyParams );
+	void OnAnimatedVectorPropertyChanged( const FKeyPropertyParams& KeyPropertyParams );
 	
 	/**
 	 * Called by the details panel when an animatable color property changes
@@ -55,7 +57,7 @@ private:
 	 * @param InObjectsThatChanged	List of objects that changed
 	 * @param PropertyValue			Handle to the property value which changed
 	 */
-	void OnAnimatedColorPropertyChanged( const struct FKeyPropertyParams& KeyPropertyParams );
+	void OnAnimatedColorPropertyChanged( const FKeyPropertyParams& KeyPropertyParams );
 	
 	/** Delegates for animatable property changed functions */
 	template <typename Type, typename TrackType>
diff --git a/Engine/Source/Editor/UMGEditor/Private/Animation/UMGDetailKeyframeHandler.cpp b/Engine/Source/Editor/UMGEditor/Private/Animation/UMGDetailKeyframeHandler.cpp
--- a/Engine/Source/Editor/UMGEditor/Private/Animation/UMGDetailKeyframeHandler.cpp
+++ b/Engine/Source/Editor/UMGEditor/Private/Animation/UMGDetailKeyframeHandler.cpp
@@ -2,7 +2,8 @@
 
 #include "UMGEditorPrivatePCH.h"
 #include "UMGDetailKeyframeHandler.h"
-
+#include "ISequencer.h"
+#include "PropertyHandle.h"
 
 FUMGDetailKeyframeHandler::FUMGDetailKeyframeHandler(TSharedPtr<FWidgetBlueprintEditor> InBlueprintEditor)
 	: BlueprintEditor( InBlueprintEditor )
diff --git a/Engine/Source/Editor/UMGEditor/Private/Details/SWidgetDetailsView.cpp b/Engine/Source/Editor/UMGEditor/Private/Details/SWidgetDetailsView.cpp
--- a/Engine/Source/Editor/UMGEditor/Private/Details/SWidgetDetailsView.cpp
+++ b/Engine/Source/Editor/UMGEditor/Private/Details/SWidgetDetailsView.cpp
@@ -6,6 +6,8 @@
 #include "BlueprintEditor.h"
 #include "IDetailsView.h"
 #include "Kismet2NameValidators.h"
+#include "ISequencer.h"
+#include "Animation/UMGDetailKeyframeHandler.h"
 
 #define LOCTEXT_NAMESPACE "UMG"
 
