Sequencer - renamed Transform track to 3D transform track to not be confused with 2D transforms

diff --git a/Engine/Source/Developer/MovieSceneTools/Private/MovieSceneToolsModule.cpp b/Engine/Source/Developer/MovieSceneTools/Private/MovieSceneToolsModule.cpp
--- a/Engine/Source/Developer/MovieSceneTools/Private/MovieSceneToolsModule.cpp
+++ b/Engine/Source/Developer/MovieSceneTools/Private/MovieSceneToolsModule.cpp
@@ -38,7 +38,7 @@ class FMovieSceneToolsModule : public IMovieSceneTools
 		// Register with the sequencer module that we provide auto-key handlers.
 		ISequencerModule& SequencerModule = FModuleManager::Get().LoadModuleChecked<ISequencerModule>( "Sequencer" );
 		SequencerModule.RegisterTrackEditor( FOnCreateTrackEditor::CreateStatic( &FPropertyTrackEditor::CreateTrackEditor ) );
-		SequencerModule.RegisterTrackEditor( FOnCreateTrackEditor::CreateStatic( &FTransformTrackEditor::CreateTrackEditor ) );
+		SequencerModule.RegisterTrackEditor( FOnCreateTrackEditor::CreateStatic( &F3DTransformTrackEditor::CreateTrackEditor ) );
 		SequencerModule.RegisterTrackEditor( FOnCreateTrackEditor::CreateStatic( &FDirectorTrackEditor::CreateTrackEditor ) );
 		SequencerModule.RegisterTrackEditor( FOnCreateTrackEditor::CreateStatic( &FSubMovieSceneTrackEditor::CreateTrackEditor ) );
 		SequencerModule.RegisterTrackEditor( FOnCreateTrackEditor::CreateStatic( &FAudioTrackEditor::CreateTrackEditor ) );
@@ -63,7 +63,7 @@ class FMovieSceneToolsModule : public IMovieSceneTools
 			// Unregister auto key handlers
 			ISequencerModule& SequencerModule = FModuleManager::Get().GetModuleChecked<ISequencerModule>( "Sequencer" );
 			SequencerModule.UnRegisterTrackEditor( FOnCreateTrackEditor::CreateStatic( &FPropertyTrackEditor::CreateTrackEditor ) );
-			SequencerModule.UnRegisterTrackEditor( FOnCreateTrackEditor::CreateStatic( &FTransformTrackEditor::CreateTrackEditor ) );
+			SequencerModule.UnRegisterTrackEditor( FOnCreateTrackEditor::CreateStatic( &F3DTransformTrackEditor::CreateTrackEditor ) );
 			SequencerModule.UnRegisterTrackEditor( FOnCreateTrackEditor::CreateStatic( &FDirectorTrackEditor::CreateTrackEditor ) );
 			SequencerModule.UnRegisterTrackEditor( FOnCreateTrackEditor::CreateStatic( &FSubMovieSceneTrackEditor::CreateTrackEditor ) );
 			SequencerModule.UnRegisterTrackEditor( FOnCreateTrackEditor::CreateStatic( &FAudioTrackEditor::CreateTrackEditor ) );
diff --git a/Engine/Source/Developer/MovieSceneTools/Private/TransformTrackEditor.cpp b/Engine/Source/Developer/MovieSceneTools/Private/TransformTrackEditor.cpp
--- a/Engine/Source/Developer/MovieSceneTools/Private/TransformTrackEditor.cpp
+++ b/Engine/Source/Developer/MovieSceneTools/Private/TransformTrackEditor.cpp
@@ -4,8 +4,8 @@
 #include "ScopedTransaction.h"
 #include "MovieScene.h"
 #include "MovieSceneSection.h"
-#include "MovieSceneTransformTrack.h"
-#include "MovieSceneTransformSection.h"
+#include "MovieScene3DTransformTrack.h"
+#include "MovieScene3DTransformSection.h"
 #include "ISequencerObjectChangeListener.h"
 #include "ISequencerSection.h"
 #include "ISectionLayoutBuilder.h"
@@ -17,10 +17,10 @@
 /**
  * Class that draws a transform section in the sequencer
  */
-class FTransformSection : public ISequencerSection
+class F3DTransformSection : public ISequencerSection
 {
 public:
-	FTransformSection( UMovieSceneSection& InSection )
+	F3DTransformSection( UMovieSceneSection& InSection )
 		: Section( InSection )
 	{
 	}
@@ -40,7 +40,7 @@ public:
 
 	virtual void GenerateSectionLayout( class ISectionLayoutBuilder& LayoutBuilder ) const override
 	{
-		UMovieSceneTransformSection* TransformSection = Cast<UMovieSceneTransformSection>( &Section );
+		UMovieScene3DTransformSection* TransformSection = Cast<UMovieScene3DTransformSection>( &Section );
 
 		// This generates the tree structure for the transform section
 		LayoutBuilder.PushCategory( "Location", NSLOCTEXT("FTransformSection", "LocationArea", "Location") );
@@ -81,42 +81,42 @@ private:
 	UMovieSceneSection& Section;
 };
 
-FTransformTrackEditor::FTransformTrackEditor( TSharedRef<ISequencer> InSequencer )
+F3DTransformTrackEditor::F3DTransformTrackEditor( TSharedRef<ISequencer> InSequencer )
 	: FMovieSceneTrackEditor( InSequencer ) 
 {
 	// Listen for actor/component movement
-	GEditor->OnBeginObjectMovement().AddRaw( this, &FTransformTrackEditor::OnPreTransformChanged );
-	GEditor->OnEndObjectMovement().AddRaw( this, &FTransformTrackEditor::OnTransformChanged );
+	GEditor->OnBeginObjectMovement().AddRaw( this, &F3DTransformTrackEditor::OnPreTransformChanged );
+	GEditor->OnEndObjectMovement().AddRaw( this, &F3DTransformTrackEditor::OnTransformChanged );
 }
 
-FTransformTrackEditor::~FTransformTrackEditor()
+F3DTransformTrackEditor::~F3DTransformTrackEditor()
 {
 	GEditor->OnBeginObjectMovement().RemoveAll( this );
 	GEditor->OnEndObjectMovement().RemoveAll( this );
 }
 
-TSharedRef<FMovieSceneTrackEditor> FTransformTrackEditor::CreateTrackEditor( TSharedRef<ISequencer> InSequencer )
+TSharedRef<FMovieSceneTrackEditor> F3DTransformTrackEditor::CreateTrackEditor( TSharedRef<ISequencer> InSequencer )
 {
-	return MakeShareable( new FTransformTrackEditor( InSequencer ) );
+	return MakeShareable( new F3DTransformTrackEditor( InSequencer ) );
 }
 
-bool FTransformTrackEditor::SupportsType( TSubclassOf<UMovieSceneTrack> Type ) const
+bool F3DTransformTrackEditor::SupportsType( TSubclassOf<UMovieSceneTrack> Type ) const
 {
 	// We support animatable transforms
-	return Type == UMovieSceneTransformTrack::StaticClass();
+	return Type == UMovieScene3DTransformTrack::StaticClass();
 }
 
-TSharedRef<ISequencerSection> FTransformTrackEditor::MakeSectionInterface( UMovieSceneSection& SectionObject, UMovieSceneTrack* Track )
+TSharedRef<ISequencerSection> F3DTransformTrackEditor::MakeSectionInterface( UMovieSceneSection& SectionObject, UMovieSceneTrack* Track )
 {
 	check( SupportsType( SectionObject.GetOuter()->GetClass() ) );
 
-	TSharedRef<ISequencerSection> NewSection( new FTransformSection( SectionObject ) );
+	TSharedRef<ISequencerSection> NewSection( new F3DTransformSection( SectionObject ) );
 
 	return NewSection;
 }
 
 
-void FTransformTrackEditor::OnPreTransformChanged( UObject& InObject )
+void F3DTransformTrackEditor::OnPreTransformChanged( UObject& InObject )
 {
 	UMovieScene* MovieScene = GetMovieScene();
 	float AutoKeyTime = GetTimeForKey( MovieScene );
@@ -160,7 +160,7 @@ struct FTransformDataPair
 	FTransformData LastTransformData;
 };
 
-void FTransformTrackEditor::OnTransformChanged( UObject& InObject )
+void F3DTransformTrackEditor::OnTransformChanged( UObject& InObject )
 {
 	const TSharedPtr<ISequencer> Sequencer = GetSequencer();
 
@@ -202,12 +202,12 @@ void FTransformTrackEditor::OnTransformChanged( UObject& InObject )
 
 		FTransformDataPair TransformPair(NewTransformData, ExistingTransform);
 
-		AnimatablePropertyChanged(UMovieSceneTransformTrack::StaticClass(), true,
-			FOnKeyProperty::CreateRaw(this, &FTransformTrackEditor::OnTransformChangedInternals, &InObject, ObjectHandle, TransformPair, true));
+		AnimatablePropertyChanged(UMovieScene3DTransformTrack::StaticClass(), true,
+			FOnKeyProperty::CreateRaw(this, &F3DTransformTrackEditor::OnTransformChangedInternals, &InObject, ObjectHandle, TransformPair, true));
 	}
 }
 
-void FTransformTrackEditor::AddKey(const FGuid& ObjectGuid, UObject* AdditionalAsset)
+void F3DTransformTrackEditor::AddKey(const FGuid& ObjectGuid, UObject* AdditionalAsset)
 {
 	TArray<UObject*> OutObjects;
 	GetSequencer()->GetRuntimeObjects( GetSequencer()->GetFocusedMovieSceneInstance(), ObjectGuid, OutObjects);
@@ -240,13 +240,13 @@ void FTransformTrackEditor::AddKey(const FGuid& ObjectGuid, UObject* AdditionalA
 
 			FTransformDataPair TransformPair(CurrentTransform, FTransformData());
 
-			AnimatablePropertyChanged(UMovieSceneTransformTrack::StaticClass(), false,
-				FOnKeyProperty::CreateRaw(this, &FTransformTrackEditor::OnTransformChangedInternals, Object, ObjectHandle, TransformPair, false));
+			AnimatablePropertyChanged(UMovieScene3DTransformTrack::StaticClass(), false,
+				FOnKeyProperty::CreateRaw(this, &F3DTransformTrackEditor::OnTransformChangedInternals, Object, ObjectHandle, TransformPair, false));
 		}
 	}
 }
 
-void FTransformTrackEditor::OnTransformChangedInternals(float KeyTime, UObject* InObject, FGuid ObjectHandle, FTransformDataPair TransformPair, bool bAutoKeying)
+void F3DTransformTrackEditor::OnTransformChangedInternals(float KeyTime, UObject* InObject, FGuid ObjectHandle, FTransformDataPair TransformPair, bool bAutoKeying)
 {
 	// Only unwind rotation if we're generating keys while recording (scene is actively playing back)
 	const bool bUnwindRotation = GetSequencer()->IsRecordingLive();
@@ -254,8 +254,8 @@ void FTransformTrackEditor::OnTransformChangedInternals(float KeyTime, UObject*
 	FName Transform("Transform");
 	if (ObjectHandle.IsValid())
 	{
-		UMovieSceneTrack* Track = GetTrackForObject( ObjectHandle, UMovieSceneTransformTrack::StaticClass(), Transform );
-		UMovieSceneTransformTrack* TransformTrack = CastChecked<UMovieSceneTransformTrack>( Track );
+		UMovieSceneTrack* Track = GetTrackForObject( ObjectHandle, UMovieScene3DTransformTrack::StaticClass(), Transform );
+		UMovieScene3DTransformTrack* TransformTrack = CastChecked<UMovieScene3DTransformTrack>( Track );
 		// Transform name and path are the same
 		TransformTrack->SetPropertyNameAndPath( Transform, Transform.ToString() );
 	
diff --git a/Engine/Source/Developer/MovieSceneTools/Private/TransformTrackEditor.h b/Engine/Source/Developer/MovieSceneTools/Private/TransformTrackEditor.h
--- a/Engine/Source/Developer/MovieSceneTools/Private/TransformTrackEditor.h
+++ b/Engine/Source/Developer/MovieSceneTools/Private/TransformTrackEditor.h
@@ -2,12 +2,12 @@
 
 #pragma once
 
-#include "MovieSceneTransformTrack.h"
+#include "MovieScene3DTransformTrack.h"
 
 /**
  * Tools for animatable transforms
  */
-class FTransformTrackEditor : public FMovieSceneTrackEditor
+class F3DTransformTrackEditor : public FMovieSceneTrackEditor
 {
 public:
 	/**
@@ -15,8 +15,8 @@ public:
 	 *
 	 * @param InSequencer	The sequencer instance to be used by this tool
 	 */
-	FTransformTrackEditor( TSharedRef<ISequencer> InSequencer );
-	~FTransformTrackEditor();
+	F3DTransformTrackEditor( TSharedRef<ISequencer> InSequencer );
+	~F3DTransformTrackEditor();
 
 	/**
 	 * Creates an instance of this class.  Called by a sequencer 
diff --git a/Engine/Source/Runtime/MovieSceneCoreTypes/Classes/MovieScene3DTransformSection.h b/Engine/Source/Runtime/MovieSceneCoreTypes/Classes/MovieScene3DTransformSection.h
new file mode 100644
--- /dev/null
+++ b/Engine/Source/Runtime/MovieSceneCoreTypes/Classes/MovieScene3DTransformSection.h
@@ -0,0 +1,124 @@
+// Copyright 1998-2014 Epic Games, Inc. All Rights Reserved.
+
+#pragma once
+
+#include "MovieSceneSection.h"
+
+#include "MovieScene3DTransformSection.generated.h"
+
+/**
+ * A 3D transform section
+ */
+UCLASS(MinimalAPI)
+class UMovieScene3DTransformSection : public UMovieSceneSection
+{
+	GENERATED_UCLASS_BODY()
+public:
+
+	/** MovieSceneSection interface */
+	virtual void MoveSection( float DeltaPosition ) override;
+	virtual void DilateSection( float DilationFactor, float Origin ) override;
+
+	/**
+	 * Evaluates the translation component of the transform
+	 *
+	 * @param Time				The position in time within the movie scene
+	 * @param OutTranslation	The evaluated translation.  Note: will remain unchanged if there were no keys to evaluate
+	 * @return true if there were any keys to evaluate.  
+	 */
+	bool EvalTranslation( float Time, FVector& OutTranslation ) const;
+
+	/**
+	 * Evaluates the rotation component of the transform
+	 *
+	 * @param Time				The position in time within the movie scene
+	 * @param OutRotation		The evaluated rotation.  Note: will remain unchanged if there were no keys to evaluate
+	 * @return true if there were any keys to evaluate.  
+	 */
+	bool EvalRotation( float Time, FRotator& OutRotation ) const;
+
+	/**
+	 * Evaluates the scale component of the transform
+	 *
+	 * @param Time				The position in time within the movie scene
+	 * @param OutScale			The evaluated scale.  Note: will remain unchanged if there were no keys to evaluate
+	 * @return true if there were any keys to evaluate.  
+	 */
+	bool EvalScale( float Time, FVector& OutScale ) const;
+
+	/** 
+	 * Returns the translation curve for a specific axis
+	 *
+	 * @param Axis	The axis of the translation curve to get
+	 * @return The curve on the axis
+	 */
+	virtual FRichCurve& GetTranslationCurve( EAxis::Type Axis );
+
+	/** 
+	 * Returns the rotation curve for a specific axis
+	 *
+	 * @param Axis	The axis of the rotation curve to get
+	 * @return The curve on the axis
+	 */
+	virtual FRichCurve& GetRotationCurve( EAxis::Type Axis );
+
+	/** 
+	 * Returns the scale curve for a specific axis
+	 *
+	 * @param Axis	The axis of the scale curve to get
+	 * @return The curve on the axis
+	 */
+	virtual FRichCurve& GetScaleCurve( EAxis::Type Axis );
+
+	/**
+	 * Adds keys each component of translation. Note: Only adds keys if a value has changed
+	 *
+	 * @param TransformKey	Information about the transform that should be keyed
+	 */
+	void AddTranslationKeys( const class FTransformKey& TransformKey );
+
+	/**
+	 * Adds keys each component of rotation. Note: Only adds keys if a value has changed
+	 *
+	 * @param TransformKey	Information about the transform that should be keyed
+	 * @param bUnwindRotation	When true, the value will be treated like a rotation value in degrees, and will automatically be unwound to prevent flipping 360 degrees from the previous key 	 
+	 */
+	void AddRotationKeys( const class FTransformKey& TransformKey, const bool bUnwindRotation );
+
+	/**
+	 * Adds keys each component of scale. Note: Only adds keys if a value has changed
+	 *
+	 * @param TransformKey	Information about the transform that should be keyed
+	 */
+	void AddScaleKeys( const class FTransformKey& TransformKey );
+	
+	/** 
+	 * Determines if a new key would be new data, or just a duplicate of existing data
+	 *
+	 * @param TransformKey The new key that would be added
+	 * @return True if the new key would be new data, false if duplicate
+	 */
+	bool NewKeyIsNewData(const class FTransformKey& TransformKey) const;
+private:
+	/**
+	 * Adds a key to a rich curve, finding an existing key to modify or adding a new one
+	 *
+	 * @param InCurve	The curve to add keys to
+	 * @param Time		The time where the key should be added
+	 * @param Value		The value at the given time
+	 * @param	bUnwindRotation		When true, the value will be treated like a rotation value in degrees, and will automatically be unwound to prevent flipping 360 degrees from the previous key 
+	 */
+	void AddKeyToCurve( FRichCurve& InCurve, float Time, float Value, const bool bUnwindRotation );
+private:
+	/** Translation curves */
+	UPROPERTY()
+	FRichCurve Translation[3];
+	
+	/** Rotation curves */
+	UPROPERTY()
+	FRichCurve Rotation[3];
+
+	/** Scale curves */
+	UPROPERTY()
+	FRichCurve Scale[3];
+};
\ No newline at end of file
diff --git a/Engine/Source/Runtime/MovieSceneCoreTypes/Classes/MovieScene3DTransformTrack.h b/Engine/Source/Runtime/MovieSceneCoreTypes/Classes/MovieScene3DTransformTrack.h
new file mode 100644
--- /dev/null
+++ b/Engine/Source/Runtime/MovieSceneCoreTypes/Classes/MovieScene3DTransformTrack.h
@@ -0,0 +1,162 @@
+// Copyright 1998-2014 Epic Games, Inc. All Rights Reserved.
+
+#pragma once
+
+#include "MovieScene.h"
+#include "MovieSceneTrack.h"
+#include "MovieScenePropertyTrack.h"
+#include "MovieScene3DTransformTrack.generated.h"
+
+/**
+ * Stores information about a transform for the purpose of adding keys to a transform section
+ */
+struct FTransformData
+{
+	/** Translation component */
+	FVector Translation;
+	/** Rotation component */
+	FRotator Rotation;
+	/** Scale component */
+	FVector Scale;
+	/** Whether or not the data is valid (any values set) */
+	bool bValid;
+
+	bool IsValid() const { return bValid; }
+
+	/**
+	 * Constructor.  Builds the data from a scene component
+	 * Uses relative transform only
+	 *
+	 * @param InComponent	The component to build from
+	 */
+	FTransformData( const USceneComponent* InComponent )
+		: Translation( InComponent->RelativeLocation )
+		, Rotation( InComponent->RelativeRotation )
+		, Scale( InComponent->RelativeScale3D )
+		, bValid( true )
+	{}
+
+	FTransformData()
+		: bValid( false )
+	{}
+};
+
+/**
+ * Represents a transform that should be keyed at a specific  time
+ * 
+ * Only keys components of the transform that have changed
+ */
+class FTransformKey
+{
+public:
+	/**
+	 * Constructor
+	 *
+	 * @param InKeyTime			The time where keys should be added
+	 * @param InNewTransform	The transform values that should be keyed
+	 * @param PreviousTransform	The previous transform values (for comparing against new values).  If this is set only changed values will be keyed
+	 * 
+	 */
+	FTransformKey( float InKeyTime, const FTransformData& InNewTransform, const FTransformData& InPreviousTransform = FTransformData() )
+		: NewTransform( InNewTransform )
+		, PreviousTransform( InPreviousTransform )
+		, KeyTime( InKeyTime )
+	{
+	}	
+
+	/** @return The time where transform keys should be added */
+	float GetKeyTime() const { return KeyTime; }
+
+	/** @return the translation value of the key */
+	const FVector& GetTranslationValue() const { return NewTransform.Translation; }
+
+	/** @return the rotation value of the key */
+	const FRotator& GetRotationValue() const { return NewTransform.Rotation; }
+
+	/** @return the scale value of the key */
+	const FVector& GetScaleValue() const { return NewTransform.Scale; }
+
+	/**
+	 * Determines if a key should be added on the given axis of translation
+	 *
+	 * @param Axis	The axis to check
+	 * @return Whether or not we should key translation on the given axis
+	 */
+	bool ShouldKeyTranslation( EAxis::Type Axis ) const;
+
+	/**
+	 * Determines if a key should be added on the given axis of rotation
+	 *
+	 * @param Axis	The axis to check
+	 * @return Whether or not we should key rotation on the given axis
+	 */
+	bool ShouldKeyRotation( EAxis::Type Axis ) const;
+
+	/**
+	 * Determines if a key should be added on the given axis of scale
+	 *
+	 * @param Axis	The axis to check
+	 * @return Whether or not we should key scale on the given axis
+	 */
+	bool ShouldKeyScale( EAxis::Type Axis ) const;
+	
+	/**
+	 * Determines if a key should be added to any axis for any type of transform
+	 *
+	 * @return Whether or not we should key on the given axis
+	 */
+	bool ShouldKeyAny() const;
+private:
+	/**
+	 * Determines whether or not a vector component is different on the given axis
+	 *
+	 * @param Axis	The axis to check
+	 * @param Current	The current vector value
+	 * @param Previous	The previous vector value
+	 * @return true if a vector component is different on the given axis
+	 */
+	bool GetVectorComponentIfDifferent( EAxis::Type Axis, const FVector& Current, const FVector& Previous ) const;
+private:
+	/** New transform data */
+	FTransformData NewTransform;
+	FTransformData PreviousTransform;
+	float KeyTime;
+};
+
+/**
+ * Handles manipulation of component transforms in a movie scene
+ */
+UCLASS( MinimalAPI )
+class UMovieScene3DTransformTrack : public UMovieScenePropertyTrack
+{
+	GENERATED_UCLASS_BODY()
+public:
+	/** UMovieSceneTrack interface */
+	virtual UMovieSceneSection* CreateNewSection() override;
+	virtual TSharedPtr<IMovieSceneTrackInstance> CreateInstance() override;
+	
+	/**
+	 * Adds a key to a section.  Will create the section if it doesn't exist
+	 *
+	 * @param ObjectHandle		Handle to the object(s) being changed
+	 * @param InKey				The transform key to add
+	 * @param bUnwindRotation	When true, the value will be treated like a rotation value in degrees, and will automatically be unwound to prevent flipping 360 degrees from the previous key 	 
+	 * @return True if the key was successfully added.
+	 */
+	virtual bool AddKeyToSection( const FGuid& ObjectHandle, const FTransformKey& InKey, const bool bUnwindRotation );
+
+	/**
+	 * Evaluates the track at the playback position
+	 *
+	 * @param Position		The current playback position
+	 * @param LastPosition		The last plackback position
+	 * @param OutTranslation	The evalulated translation component of the transform
+	 * @param OutRotation		The evalulated rotation component of the transform
+	 * @param OutScale		The evalulated scale component of the transform
+	 * @param OutHasTranslationKeys true if a translation key was evaluated 
+	 * @param OutHasRotationKeys 	true if a rotation key was evaluated
+	 * @param OutHasScaleKeys 	true if a scale key was evaluated
+	 * @return true if anything was evaluated. Note: if false is returned OutBool remains unchanged
+	 */
+	virtual bool Eval( float Position, float LastPosition, FVector& OutTranslation, FRotator& OutRotation, FVector& OutScale, bool& OutHasTranslationKeys, bool& OutHasRotationKeys, bool& OutHasScaleKeys ) const;
+};
diff --git a/Engine/Source/Runtime/MovieSceneCoreTypes/Classes/MovieSceneTransformSection.h b/Engine/Source/Runtime/MovieSceneCoreTypes/Classes/MovieSceneTransformSection.h
deleted file mode 100644
diff --git a/Engine/Source/Runtime/MovieSceneCoreTypes/Classes/MovieSceneTransformTrack.h b/Engine/Source/Runtime/MovieSceneCoreTypes/Classes/MovieSceneTransformTrack.h
deleted file mode 100644
diff --git a/Engine/Source/Runtime/MovieSceneCoreTypes/Private/Instances/MovieScene3DTransformTrackInstance.cpp b/Engine/Source/Runtime/MovieSceneCoreTypes/Private/Instances/MovieScene3DTransformTrackInstance.cpp
new file mode 100644
--- /dev/null
+++ b/Engine/Source/Runtime/MovieSceneCoreTypes/Private/Instances/MovieScene3DTransformTrackInstance.cpp
@@ -0,0 +1,53 @@
+// Copyright 1998-2014 Epic Games, Inc. All Rights Reserved.
+
+#include "MovieSceneCoreTypesPCH.h"
+#include "MovieScene3DTransformTrackInstance.h"
+
+
+FMovieScene3DTransformTrackInstance::FMovieScene3DTransformTrackInstance( UMovieScene3DTransformTrack& InTransformTrack )
+{
+	TransformTrack = &InTransformTrack;
+}
+
+void FMovieScene3DTransformTrackInstance::Update( float Position, float LastPosition, const TArray<UObject*>& RuntimeObjects, class IMovieScenePlayer& Player ) 
+{
+	FVector Translation;
+	FRotator Rotation;
+	FVector Scale;
+	bool bHasTranslationKeys = false;
+	bool bHasRotationKeys = false;
+	bool bHasScaleKeys = false;
+
+	if( TransformTrack->Eval( Position, LastPosition, Translation, Rotation, Scale, bHasTranslationKeys, bHasRotationKeys, bHasScaleKeys ) )
+	{
+		for( int32 ObjIndex = 0; ObjIndex < RuntimeObjects.Num(); ++ObjIndex )
+		{
+			UObject* Object = RuntimeObjects[ObjIndex];
+
+			AActor* Actor = Cast<AActor>( Object );
+
+			USceneComponent* SceneComponent = NULL;
+			if( Actor && Actor->GetRootComponent() )
+			{
+				// If there is an actor, modify its root component
+				SceneComponent = Actor->GetRootComponent();
+			}
+			else
+			{
+				// No actor was found.  Attempt to get the object as a component in the case that we are editing them directly.
+				SceneComponent = Cast<USceneComponent>( Object );
+			}
+
+			// Set the relative translation and rotation.  Note they are set once instead of individually to avoid a redundant component transform update.
+			SceneComponent->SetRelativeLocationAndRotation(
+				bHasTranslationKeys ? Translation : SceneComponent->RelativeLocation,
+				bHasRotationKeys ? Rotation : SceneComponent->RelativeRotation );
+
+			if( bHasScaleKeys )
+			{
+				SceneComponent->SetRelativeScale3D( Scale );
+			}
+		}
+	}
+}
+
diff --git a/Engine/Source/Runtime/MovieSceneCoreTypes/Private/Instances/MovieScene3DTransformTrackInstance.h b/Engine/Source/Runtime/MovieSceneCoreTypes/Private/Instances/MovieScene3DTransformTrackInstance.h
new file mode 100644
--- /dev/null
+++ b/Engine/Source/Runtime/MovieSceneCoreTypes/Private/Instances/MovieScene3DTransformTrackInstance.h
@@ -0,0 +1,19 @@
+// Copyright 1998-2014 Epic Games, Inc. All Rights Reserved.
+
+#pragma once
+
+/**
+ * Instance of a UMovieSceneTransformTrack
+ */
+class FMovieScene3DTransformTrackInstance : public IMovieSceneTrackInstance
+{
+public:
+	FMovieScene3DTransformTrackInstance( UMovieScene3DTransformTrack& InTransformTrack );
+
+	/** IMovieSceneTrackInstance interface */
+	virtual void Update( float Position, float LastPosition, const TArray<UObject*>& RuntimeObjects, class IMovieScenePlayer& Player ) override;
+	virtual void RefreshInstance( const TArray<UObject*>& RuntimeObjects, class IMovieScenePlayer& Player ) override {}
+private:
+	/** Track that is being instanced */
+	UMovieScene3DTransformTrack* TransformTrack;
+};
diff --git a/Engine/Source/Runtime/MovieSceneCoreTypes/Private/Instances/MovieSceneTransformTrackInstance.cpp b/Engine/Source/Runtime/MovieSceneCoreTypes/Private/Instances/MovieSceneTransformTrackInstance.cpp
deleted file mode 100644
diff --git a/Engine/Source/Runtime/MovieSceneCoreTypes/Private/Instances/MovieSceneTransformTrackInstance.h b/Engine/Source/Runtime/MovieSceneCoreTypes/Private/Instances/MovieSceneTransformTrackInstance.h
deleted file mode 100644
diff --git a/Engine/Source/Runtime/MovieSceneCoreTypes/Private/MovieScene3DTransformSection.cpp b/Engine/Source/Runtime/MovieSceneCoreTypes/Private/MovieScene3DTransformSection.cpp
new file mode 100644
--- /dev/null
+++ b/Engine/Source/Runtime/MovieSceneCoreTypes/Private/MovieScene3DTransformSection.cpp
@@ -0,0 +1,212 @@
+// Copyright 1998-2014 Epic Games, Inc. All Rights Reserved.
+
+#include "MovieSceneCoreTypesPCH.h"
+#include "MovieScene3DTransformSection.h"
+#include "MovieScene3DTransformTrack.h"
+
+UMovieScene3DTransformSection::UMovieScene3DTransformSection( const FPostConstructInitializeProperties& PCIP )
+	: Super( PCIP )
+{
+
+}
+
+void UMovieScene3DTransformSection::MoveSection( float DeltaTime )
+{
+	Super::MoveSection( DeltaTime );
+
+	// Move all the curves in this section
+	for( int32 Axis = 0; Axis < 3; ++Axis )
+	{
+		Translation[Axis].ShiftCurve( DeltaTime );
+		Rotation[Axis].ShiftCurve( DeltaTime );
+		Scale[Axis].ShiftCurve( DeltaTime );
+	}
+}
+
+void UMovieScene3DTransformSection::DilateSection( float DilationFactor, float Origin )
+{
+	Super::DilateSection(DilationFactor, Origin);
+
+	for( int32 Axis = 0; Axis < 3; ++Axis )
+	{
+		Translation[Axis].ScaleCurve( Origin, DilationFactor );
+		Rotation[Axis].ScaleCurve( Origin, DilationFactor );
+		Scale[Axis].ScaleCurve( Origin, DilationFactor );
+	}
+}
+
+	
+
+bool UMovieScene3DTransformSection::EvalTranslation( float Time, FVector& OutTranslation ) const
+{
+	bool bHasTranslationKeys = Translation[0].GetNumKeys() > 0 || Translation[1].GetNumKeys() > 0 || Translation[2].GetNumKeys() > 0;
+
+	if( bHasTranslationKeys )
+	{
+		// Evaluate each component of translation if there are any keys
+		OutTranslation.X = Translation[0].Eval( Time );
+		OutTranslation.Y = Translation[1].Eval( Time );
+		OutTranslation.Z = Translation[2].Eval( Time );
+	}
+
+	return bHasTranslationKeys;
+}
+
+bool UMovieScene3DTransformSection::EvalRotation( float Time, FRotator& OutRotation ) const
+{
+	bool bHasRotationKeys = Rotation[0].GetNumKeys() > 0 || Rotation[1].GetNumKeys() > 0 || Rotation[2].GetNumKeys() > 0;
+
+	if( bHasRotationKeys )
+	{
+		// Evaluate each component of rotation if there are any keys
+		OutRotation.Roll = Rotation[0].Eval( Time );
+		OutRotation.Pitch = Rotation[1].Eval( Time );
+		OutRotation.Yaw = Rotation[2].Eval( Time );
+	}
+
+	return bHasRotationKeys;
+
+}
+
+bool UMovieScene3DTransformSection::EvalScale( float Time, FVector& OutScale ) const
+{
+	bool bHasScaleKeys = Scale[0].GetNumKeys() > 0 || Scale[1].GetNumKeys() > 0 || Scale[2].GetNumKeys() > 0;
+	
+	if( bHasScaleKeys )
+	{
+		// Evaluate each component of scale if there are any keys
+		OutScale.X = Scale[0].Eval( Time );
+		OutScale.Y = Scale[1].Eval( Time );
+		OutScale.Z = Scale[2].Eval( Time );
+	}
+
+	return bHasScaleKeys;
+
+}
+
+/**
+ * Chooses an appropriate curve from an axis and a set of curves
+ */
+static FRichCurve* ChooseCurve( EAxis::Type Axis, FRichCurve* Curves )
+{
+	switch( Axis )
+	{
+	case EAxis::X:
+		return &Curves[0];
+		break;
+	case EAxis::Y:
+		return &Curves[1];
+		break;
+	case EAxis::Z:
+		return &Curves[2];
+		break;
+	default:
+		check( false );
+		return NULL;
+		break;
+	}
+}
+
+FRichCurve& UMovieScene3DTransformSection::GetTranslationCurve( EAxis::Type Axis ) 
+{
+	return *ChooseCurve( Axis, Translation );
+}
+
+FRichCurve& UMovieScene3DTransformSection::GetRotationCurve( EAxis::Type Axis )
+{
+	return *ChooseCurve( Axis, Rotation );
+}
+
+FRichCurve& UMovieScene3DTransformSection::GetScaleCurve( EAxis::Type Axis )
+{
+	return *ChooseCurve( Axis, Scale );
+}
+
+void UMovieScene3DTransformSection::AddKeyToCurve( FRichCurve& InCurve, float Time, float Value, const bool bUnwindRotation )
+{
+	if( IsTimeWithinSection(Time) )
+	{
+		Modify();
+		InCurve.UpdateOrAddKey(Time, Value);
+	}
+}
+
+void UMovieScene3DTransformSection::AddTranslationKeys( const FTransformKey& TransformKey )
+{
+	const float Time = TransformKey.GetKeyTime();
+
+	const bool bUnwindRotation = false; // Not needed for translation
+
+	// Key each component.  Note: we always key each component if no key exists
+	if( Translation[0].GetNumKeys() == 0 || TransformKey.ShouldKeyTranslation( EAxis::X ) )
+	{
+		AddKeyToCurve( Translation[0], Time, TransformKey.GetTranslationValue().X, bUnwindRotation );
+	}
+
+	if( Translation[1].GetNumKeys() == 0 || TransformKey.ShouldKeyTranslation( EAxis::Y ) )
+	{ 
+		AddKeyToCurve( Translation[1], Time, TransformKey.GetTranslationValue().Y, bUnwindRotation );
+	}
+
+	if( Translation[2].GetNumKeys() == 0 || TransformKey.ShouldKeyTranslation( EAxis::Z ) )
+	{ 
+		AddKeyToCurve( Translation[2], Time, TransformKey.GetTranslationValue().Z, bUnwindRotation );
+	}
+}
+
+void UMovieScene3DTransformSection::AddRotationKeys( const FTransformKey& TransformKey, const bool bUnwindRotation )
+{
+	const float Time = TransformKey.GetKeyTime();
+
+	// Key each component.  Note: we always key each component if no key exists
+	if( Rotation[0].GetNumKeys() == 0 || TransformKey.ShouldKeyRotation( EAxis::X ) )
+	{
+		AddKeyToCurve( Rotation[0], Time, TransformKey.GetRotationValue().Roll, bUnwindRotation );
+	}
+
+	if( Rotation[1].GetNumKeys() == 0 || TransformKey.ShouldKeyRotation( EAxis::Y ) )
+	{ 
+		AddKeyToCurve( Rotation[1], Time, TransformKey.GetRotationValue().Pitch, bUnwindRotation );
+	}
+
+	if( Rotation[2].GetNumKeys() == 0 || TransformKey.ShouldKeyRotation( EAxis::Z ) )
+	{ 
+		AddKeyToCurve( Rotation[2], Time, TransformKey.GetRotationValue().Yaw, bUnwindRotation );
+	}
+}
+
+void UMovieScene3DTransformSection::AddScaleKeys( const FTransformKey& TransformKey )
+{
+	const float Time = TransformKey.GetKeyTime();
+
+	const bool bUnwindRotation = false; // Not needed for scale
+
+	// Key each component.  Note: we always key each component if no key exists
+	if( Scale[0].GetNumKeys() == 0 || TransformKey.ShouldKeyScale( EAxis::X ) )
+	{
+		AddKeyToCurve( Scale[0], Time, TransformKey.GetScaleValue().X, bUnwindRotation );
+	}
+
+	if( Scale[1].GetNumKeys() == 0 || TransformKey.ShouldKeyScale( EAxis::Y ) )
+	{ 
+		AddKeyToCurve( Scale[1], Time, TransformKey.GetScaleValue().Y, bUnwindRotation );
+	}
+
+	if( Scale[2].GetNumKeys() == 0 || TransformKey.ShouldKeyScale( EAxis::Z ) )
+	{ 
+		AddKeyToCurve( Scale[2], Time, TransformKey.GetScaleValue().Z, bUnwindRotation );
+	}
+}
+
+bool UMovieScene3DTransformSection::NewKeyIsNewData(const FTransformKey& TransformKey) const
+{
+	bool bHasEmptyKeys = false;
+	for (int32 i = 0; i < 3; ++i)
+	{
+		bHasEmptyKeys = bHasEmptyKeys ||
+			Translation[i].GetNumKeys() == 0 ||
+			Rotation[i].GetNumKeys() == 0 ||
+			Scale[i].GetNumKeys() == 0;
+	}
+	return bHasEmptyKeys || TransformKey.ShouldKeyAny();
+}
diff --git a/Engine/Source/Runtime/MovieSceneCoreTypes/Private/MovieScene3DTransformTrack.cpp b/Engine/Source/Runtime/MovieSceneCoreTypes/Private/MovieScene3DTransformTrack.cpp
new file mode 100644
--- /dev/null
+++ b/Engine/Source/Runtime/MovieSceneCoreTypes/Private/MovieScene3DTransformTrack.cpp
@@ -0,0 +1,110 @@
+// Copyright 1998-2014 Epic Games, Inc. All Rights Reserved.
+
+#include "MovieSceneCoreTypesPCH.h"
+#include "MovieScene3DTransformSection.h"
+#include "MovieScene3DTransformTrack.h"
+#include "IMovieScenePlayer.h"
+#include "MovieScene3DTransformTrackInstance.h"
+
+bool FTransformKey::ShouldKeyTranslation( EAxis::Type Axis) const
+{
+	// If the previous transform is not valid we have nothing to compare against so assume we always key. Otherwise check for differences
+	return !PreviousTransform.IsValid() || GetVectorComponentIfDifferent( Axis, NewTransform.Translation, PreviousTransform.Translation );
+}
+
+bool FTransformKey::ShouldKeyRotation( EAxis::Type Axis ) const
+{
+	// If the previous transform is not valid we have nothing to compare against so assume we always key. Otherwise check for differences
+	return !PreviousTransform.IsValid() || GetVectorComponentIfDifferent( Axis, NewTransform.Rotation.Euler(), PreviousTransform.Rotation.Euler() );
+}
+
+bool FTransformKey::ShouldKeyScale( EAxis::Type Axis ) const
+{
+	// If the previous transform is not valid we have nothing to compare against so assume we always key. Otherwise check for differences
+	return !PreviousTransform.IsValid() || GetVectorComponentIfDifferent( Axis, NewTransform.Scale, PreviousTransform.Scale );
+}
+
+bool FTransformKey::ShouldKeyAny() const
+{
+	return ShouldKeyTranslation(EAxis::X) || ShouldKeyTranslation(EAxis::Y) || ShouldKeyTranslation(EAxis::Z) ||
+		ShouldKeyRotation(EAxis::X) || ShouldKeyRotation(EAxis::Y) || ShouldKeyRotation(EAxis::Z) ||
+		ShouldKeyScale(EAxis::X) || ShouldKeyScale(EAxis::Y) || ShouldKeyScale(EAxis::Z);
+}
+
+bool FTransformKey::GetVectorComponentIfDifferent( EAxis::Type Axis, const FVector& Current, const FVector& Previous ) const
+{
+	bool bShouldAddKey = false;
+
+	// Using lower precision to avoid small changes due to floating point error
+	// @todo Sequencer - Floating point error introduced by matrix math causing extra keys to be added
+	const float Precision = 1e-3;
+
+	// Add a key if there is no previous transform to compare against or the current and previous values are not equal
+
+	if( Axis == EAxis::X && !FMath::IsNearlyEqual( Current.X, Previous.X, Precision ) )
+	{
+		bShouldAddKey = true;
+	}
+	else if( Axis == EAxis::Y && !FMath::IsNearlyEqual( Current.Y, Previous.Y, Precision ) )
+	{
+		bShouldAddKey = true;
+	}
+	else if( Axis == EAxis::Z && !FMath::IsNearlyEqual( Current.Z, Previous.Z, Precision ) )
+	{
+		bShouldAddKey = true;
+	}
+
+	return bShouldAddKey;
+}
+
+UMovieScene3DTransformTrack::UMovieScene3DTransformTrack( const FPostConstructInitializeProperties& PCIP )
+	: Super( PCIP )
+{
+}
+
+UMovieSceneSection* UMovieScene3DTransformTrack::CreateNewSection()
+{
+	return ConstructObject<UMovieSceneSection>( UMovieScene3DTransformSection::StaticClass(), this, NAME_None, RF_Transactional );
+}
+
+TSharedPtr<IMovieSceneTrackInstance> UMovieScene3DTransformTrack::CreateInstance()
+{
+	return MakeShareable( new FMovieScene3DTransformTrackInstance( *this ) );
+}
+
+bool UMovieScene3DTransformTrack::AddKeyToSection( const FGuid& ObjectHandle, const FTransformKey& InKey, const bool bUnwindRotation )
+{
+	const UMovieSceneSection* NearestSection = MovieSceneHelpers::FindSectionAtTime(Sections, InKey.GetKeyTime());
+	if (!NearestSection || CastChecked<UMovieScene3DTransformSection>(NearestSection)->NewKeyIsNewData(InKey))
+	{
+		Modify();
+
+		UMovieScene3DTransformSection* NewSection = Cast<UMovieScene3DTransformSection>( FindOrAddSection( InKey.GetKeyTime() ) );
+
+		// key each component of the transform
+		NewSection->AddTranslationKeys( InKey );
+		NewSection->AddRotationKeys( InKey, bUnwindRotation );
+		NewSection->AddScaleKeys( InKey );
+
+		return true;
+	}
+	return false;
+}
+
+
+bool UMovieScene3DTransformTrack::Eval( float Position, float LastPosition, FVector& OutTranslation, FRotator& OutRotation, FVector& OutScale, bool& OutHasTranslationKeys, bool& OutHasRotationKeys, bool& OutHasScaleKeys ) const
+{
+	const UMovieSceneSection* Section = MovieSceneHelpers::FindSectionAtTime( Sections, Position );
+
+	if( Section )
+	{
+		const UMovieScene3DTransformSection* TransformSection = CastChecked<UMovieScene3DTransformSection>( Section );
+
+		// Evalulate translation,rotation, and scale curves.  If no keys were found on one of these, that component of the transform will remain unchained
+		OutHasTranslationKeys = TransformSection->EvalTranslation( Position, OutTranslation );
+		OutHasRotationKeys = TransformSection->EvalRotation( Position, OutRotation );
+		OutHasScaleKeys = TransformSection->EvalScale( Position, OutScale );
+	}
+
+	return Section != NULL;
+}
\ No newline at end of file
diff --git a/Engine/Source/Runtime/MovieSceneCoreTypes/Private/MovieSceneTransformSection.cpp b/Engine/Source/Runtime/MovieSceneCoreTypes/Private/MovieSceneTransformSection.cpp
deleted file mode 100644
diff --git a/Engine/Source/Runtime/MovieSceneCoreTypes/Private/MovieSceneTransformTrack.cpp b/Engine/Source/Runtime/MovieSceneCoreTypes/Private/MovieSceneTransformTrack.cpp
deleted file mode 100644
