Added ability for the sequencer color track to animate Slate colors

diff --git a/Engine/Source/Developer/MovieSceneTools/Private/PropertyTrackEditor.cpp b/Engine/Source/Developer/MovieSceneTools/Private/PropertyTrackEditor.cpp
--- a/Engine/Source/Developer/MovieSceneTools/Private/PropertyTrackEditor.cpp
+++ b/Engine/Source/Developer/MovieSceneTools/Private/PropertyTrackEditor.cpp
@@ -281,6 +281,8 @@ FPropertyTrackEditor::FPropertyTrackEditor( TSharedRef<ISequencer> InSequencer )
 	ObjectChangeListener.GetOnAnimatablePropertyChanged( NAME_Vector2D ).AddRaw( this, &FPropertyTrackEditor::OnAnimatedVectorPropertyChanged );
 	ObjectChangeListener.GetOnAnimatablePropertyChanged( NAME_Color ).AddRaw( this, &FPropertyTrackEditor::OnAnimatedColorPropertyChanged );
 	ObjectChangeListener.GetOnAnimatablePropertyChanged( NAME_LinearColor ).AddRaw( this, &FPropertyTrackEditor::OnAnimatedColorPropertyChanged );
+	ObjectChangeListener.GetOnAnimatablePropertyChanged( "SlateColor").AddRaw( this, &FPropertyTrackEditor::OnAnimatedColorPropertyChanged );
+
 }
 
 FPropertyTrackEditor::~FPropertyTrackEditor()
@@ -296,6 +298,7 @@ FPropertyTrackEditor::~FPropertyTrackEditor()
 		ObjectChangeListener.GetOnAnimatablePropertyChanged(NAME_Vector2D).RemoveAll(this);
 		ObjectChangeListener.GetOnAnimatablePropertyChanged(NAME_Color).RemoveAll(this);
 		ObjectChangeListener.GetOnAnimatablePropertyChanged(NAME_LinearColor).RemoveAll(this);
+		ObjectChangeListener.GetOnAnimatablePropertychanged("SlateColor").RemoveAll(this);
 	}
 }
 
@@ -424,7 +427,7 @@ void FPropertyTrackEditor::OnAnimatedVectorPropertyChanged(const FKeyPropertyPar
 
 void FPropertyTrackEditor::OnAnimatedColorPropertyChanged(const FKeyPropertyParams& KeyPropertyParams )
 {
-	bool bIsFColor = false, bIsFLinearColor = false;
+	bool bIsFColor = false, bIsFLinearColor = false, bIsSlateColor = false;
 	const UStructProperty* StructProp = Cast<const UStructProperty>(KeyPropertyParams.PropertyHandle->GetProperty());
 	if (StructProp && StructProp->Struct)
 	{
@@ -432,10 +435,14 @@ void FPropertyTrackEditor::OnAnimatedColorPropertyChanged(const FKeyPropertyPara
 
 		bIsFColor = StructName == NAME_Color;
 		bIsFLinearColor = StructName == NAME_LinearColor;
+		bIsSlateColor = StructName == FName("SlateColor");
+	}
+
+	if (!bIsFColor && !bIsFLinearColor && !bIsSlateColor) 
+	{
+		return;
 	}
-	if (!bIsFColor && !bIsFLinearColor) {return;}
 
-	check(bIsFColor ^ bIsFLinearColor);
 
 
 	FName PropertyName = KeyPropertyParams.PropertyHandle->GetProperty()->GetFName();
@@ -479,6 +486,7 @@ void FPropertyTrackEditor::OnAnimatedColorPropertyChanged(const FKeyPropertyPara
 	Key.Value = ColorValue;
 	Key.CurveName = KeyPropertyParams.InnerStructPropertyName;
 	Key.bAddKeyEvenIfUnchanged = !KeyPropertyParams.bRequireAutoKey;
+	Key.bIsSlateColor = bIsSlateColor;
 
 	AnimatablePropertyChanged(UMovieSceneColorTrack::StaticClass(), KeyPropertyParams.bRequireAutoKey, FOnKeyProperty::CreateRaw(this, &FPropertyTrackEditor::OnKeyProperty<FColorKey, UMovieSceneColorTrack>, KeyPropertyParams, Key) );
 }
diff --git a/Engine/Source/Runtime/MovieSceneCoreTypes/Classes/MovieSceneColorTrack.h b/Engine/Source/Runtime/MovieSceneCoreTypes/Classes/MovieSceneColorTrack.h
--- a/Engine/Source/Runtime/MovieSceneCoreTypes/Classes/MovieSceneColorTrack.h
+++ b/Engine/Source/Runtime/MovieSceneCoreTypes/Classes/MovieSceneColorTrack.h
@@ -12,6 +12,7 @@ struct FColorKey
 	FLinearColor Value;
 	FName CurveName;
 	bool bAddKeyEvenIfUnchanged;
+	bool bIsSlateColor;
 };
 
 /**
@@ -45,4 +46,10 @@ public:
 	 * @return true if anything was evaluated. Note: if false is returned OutColor remains unchanged
 	 */
 	virtual bool Eval( float Position, float LastPostion, FLinearColor& OutColor ) const;
+
+	bool IsSlateColor() const { return bIsSlateColor; }
+
+private:
+	UPROPERTY()
+	bool bIsSlateColor;
 };
diff --git a/Engine/Source/Runtime/MovieSceneCoreTypes/MovieSceneCoreTypes.Build.cs b/Engine/Source/Runtime/MovieSceneCoreTypes/MovieSceneCoreTypes.Build.cs
--- a/Engine/Source/Runtime/MovieSceneCoreTypes/MovieSceneCoreTypes.Build.cs
+++ b/Engine/Source/Runtime/MovieSceneCoreTypes/MovieSceneCoreTypes.Build.cs
@@ -14,6 +14,6 @@ public class MovieSceneCoreTypes : ModuleRules
 
 		PublicDependencyModuleNames.AddRange(new string[] { "Core", "CoreUObject", "Engine" });
 
-		PrivateDependencyModuleNames.Add("MovieSceneCore");
+		PrivateDependencyModuleNames.AddRange(new string[] {"MovieSceneCore", "SlateCore"} );
 	}
 }
diff --git a/Engine/Source/Runtime/MovieSceneCoreTypes/Private/Instances/MovieSceneColorTrackInstance.cpp b/Engine/Source/Runtime/MovieSceneCoreTypes/Private/Instances/MovieSceneColorTrackInstance.cpp
--- a/Engine/Source/Runtime/MovieSceneCoreTypes/Private/Instances/MovieSceneColorTrackInstance.cpp
+++ b/Engine/Source/Runtime/MovieSceneCoreTypes/Private/Instances/MovieSceneColorTrackInstance.cpp
@@ -16,10 +16,23 @@ void FMovieSceneColorTrackInstance::Update( float Position, float LastPosition,
 {
 	for(UObject* Object : RuntimeObjects)
 	{
-		FLinearColor ColorValue = PropertyBindings->GetCurrentValue<FLinearColor>(Object);
-		if(ColorTrack->Eval(Position, LastPosition, ColorValue))
+		if( ColorTrack->IsSlateColor() )
 		{
-			PropertyBindings->CallFunction(Object, &ColorValue);
+			FSlateColor ColorValue = PropertyBindings->GetCurrentValue<FSlateColor>(Object);
+			FLinearColor LinearColor = ColorValue.GetSpecifiedColor();
+			if(ColorTrack->Eval(Position, LastPosition, LinearColor))
+			{
+				FSlateColor NewColor(LinearColor);
+				PropertyBindings->CallFunction(Object, &NewColor);
+			}
+		}
+		else
+		{
+			FLinearColor ColorValue = PropertyBindings->GetCurrentValue<FLinearColor>(Object);
+			if(ColorTrack->Eval(Position, LastPosition, ColorValue))
+			{
+				PropertyBindings->CallFunction(Object, &ColorValue);
+			}
 		}
 	}
 }
diff --git a/Engine/Source/Runtime/MovieSceneCoreTypes/Private/MovieSceneColorTrack.cpp b/Engine/Source/Runtime/MovieSceneCoreTypes/Private/MovieSceneColorTrack.cpp
--- a/Engine/Source/Runtime/MovieSceneCoreTypes/Private/MovieSceneColorTrack.cpp
+++ b/Engine/Source/Runtime/MovieSceneCoreTypes/Private/MovieSceneColorTrack.cpp
@@ -8,6 +8,7 @@
 
 UMovieSceneColorTrack::UMovieSceneColorTrack( const FObjectInitializer& ObjectInitializer )
 	: Super( ObjectInitializer )
+	, bIsSlateColor( false )
 {
 }
 
@@ -24,6 +25,8 @@ TSharedPtr<IMovieSceneTrackInstance> UMovieSceneColorTrack::CreateInstance()
 
 bool UMovieSceneColorTrack::AddKeyToSection( float Time, const FColorKey& Key )
 {
+	bIsSlateColor = Key.bIsSlateColor;
+
 	const UMovieSceneSection* NearestSection = MovieSceneHelpers::FindSectionAtTime( Sections, Time );
 	if (!NearestSection || Key.bAddKeyEvenIfUnchanged || CastChecked<UMovieSceneColorSection>(NearestSection)->NewKeyIsNewData(Time, Key.Value))
 	{
