Fix compile error

diff --git a/Engine/Source/Developer/MovieSceneTools/Private/PropertyTrackEditor.cpp b/Engine/Source/Developer/MovieSceneTools/Private/PropertyTrackEditor.cpp
--- a/Engine/Source/Developer/MovieSceneTools/Private/PropertyTrackEditor.cpp
+++ b/Engine/Source/Developer/MovieSceneTools/Private/PropertyTrackEditor.cpp
@@ -298,7 +298,7 @@ FPropertyTrackEditor::~FPropertyTrackEditor()
 		ObjectChangeListener.GetOnAnimatablePropertyChanged(NAME_Vector2D).RemoveAll(this);
 		ObjectChangeListener.GetOnAnimatablePropertyChanged(NAME_Color).RemoveAll(this);
 		ObjectChangeListener.GetOnAnimatablePropertyChanged(NAME_LinearColor).RemoveAll(this);
-		ObjectChangeListener.GetOnAnimatablePropertychanged("SlateColor").RemoveAll(this);
+		ObjectChangeListener.GetOnAnimatablePropertyChanged("SlateColor").RemoveAll(this);
 	}
 }
 
