Sequencer - Specifiy a maximum value for PixelsForInput in FTimeToPixel, so that zero sized view ranges work predictably.
+ Fixes an issue where making a section too small would make it infinitely large and fixes an issue where zero sized sections could no longer be resized.

diff --git a/Engine/Source/Developer/MovieSceneTools/Public/CommonMovieSceneTools.h b/Engine/Source/Developer/MovieSceneTools/Public/CommonMovieSceneTools.h
--- a/Engine/Source/Developer/MovieSceneTools/Public/CommonMovieSceneTools.h
+++ b/Engine/Source/Developer/MovieSceneTools/Public/CommonMovieSceneTools.h
@@ -12,7 +12,7 @@ public:
 		: LocalViewRange( InLocalViewRange )
 	{
 		const float ViewRange = LocalViewRange.Size<float>();
-		PixelsPerInput = ViewRange > 0 ? AllottedGeometry.Size.X / ViewRange : 0;
+		PixelsPerInput = ViewRange > 0 ? AllottedGeometry.Size.X / ViewRange : MaxPixelsPerInput;
 	}
 
 	/**
@@ -45,8 +45,12 @@ public:
 		return PixelsPerInput;
 	}
 private:
+	static const float MaxPixelsPerInput;
+
 	/** time range of the sequencer */
 	TRange<float> LocalViewRange;
 	/** The number of pixels in the view range */
 	float PixelsPerInput;
 };
+
+const float FTimeToPixel::MaxPixelsPerInput = 10000.0f;
