Fix non-unity build error - static definition of variable in header file would cause redefinition linker errors

diff --git a/Engine/Source/Developer/MovieSceneTools/Public/CommonMovieSceneTools.h b/Engine/Source/Developer/MovieSceneTools/Public/CommonMovieSceneTools.h
--- a/Engine/Source/Developer/MovieSceneTools/Public/CommonMovieSceneTools.h
+++ b/Engine/Source/Developer/MovieSceneTools/Public/CommonMovieSceneTools.h
@@ -10,6 +10,7 @@ struct FTimeToPixel
 public:
 	FTimeToPixel( const FGeometry& AllottedGeometry, TRange<float> InLocalViewRange )
 		: LocalViewRange( InLocalViewRange )
+		, MaxPixelsPerInput(10000.0f)
 	{
 		const float ViewRange = LocalViewRange.Size<float>();
 		PixelsPerInput = ViewRange > 0 ? AllottedGeometry.Size.X / ViewRange : MaxPixelsPerInput;
@@ -45,12 +46,10 @@ public:
 		return PixelsPerInput;
 	}
 private:
-	static const float MaxPixelsPerInput;
+	const float MaxPixelsPerInput;
 
 	/** time range of the sequencer */
 	TRange<float> LocalViewRange;
 	/** The number of pixels in the view range */
 	float PixelsPerInput;
 };
-
-const float FTimeToPixel::MaxPixelsPerInput = 10000.0f;
