Fixed Clang compile error

diff --git a/Engine/Source/Developer/MovieSceneTools/Public/CommonMovieSceneTools.h b/Engine/Source/Developer/MovieSceneTools/Public/CommonMovieSceneTools.h
--- a/Engine/Source/Developer/MovieSceneTools/Public/CommonMovieSceneTools.h
+++ b/Engine/Source/Developer/MovieSceneTools/Public/CommonMovieSceneTools.h
@@ -9,8 +9,8 @@ struct FTimeToPixel
 {
 public:
 	FTimeToPixel( const FGeometry& AllottedGeometry, TRange<float> InLocalViewRange )
-		: LocalViewRange( InLocalViewRange )
-		, MaxPixelsPerInput(10000.0f)
+		: MaxPixelsPerInput(10000.0f)
+		, LocalViewRange( InLocalViewRange )
 	{
 		const float ViewRange = LocalViewRange.Size<float>();
 		PixelsPerInput = ViewRange > 0 ? AllottedGeometry.Size.X / ViewRange : MaxPixelsPerInput;
