Sequencer - Add support for byte properties, and by extension enum properties.  This enables visibility keying in UMG

diff --git a/Engine/Source/Developer/MovieSceneTools/Private/PropertyTrackEditor.cpp b/Engine/Source/Developer/MovieSceneTools/Private/PropertyTrackEditor.cpp
--- a/Engine/Source/Developer/MovieSceneTools/Private/PropertyTrackEditor.cpp
+++ b/Engine/Source/Developer/MovieSceneTools/Private/PropertyTrackEditor.cpp
@@ -7,6 +7,7 @@
 #include "MovieSceneFloatTrack.h"
 #include "MovieSceneVectorTrack.h"
 #include "MovieSceneBoolTrack.h"
+#include "MovieSceneByteTrack.h"
 #include "MovieSceneColorTrack.h"
 #include "ScopedTransaction.h"
 #include "MovieSceneSection.h"
@@ -17,6 +18,7 @@
 #include "MovieSceneFloatSection.h"
 #include "MovieSceneVectorSection.h"
 #include "MovieSceneBoolSection.h"
+#include "MovieSceneByteSection.h"
 #include "MovieSceneColorSection.h"
 #include "MovieSceneToolHelpers.h"
 #include "PropertyTrackEditor.h"
@@ -269,6 +271,22 @@ public:
 	}
 };
 
+/**
+* An implementation of byte property sections
+*/
+class FBytePropertySection : public FPropertySection
+{
+public:
+	FBytePropertySection(UMovieSceneSection& InSectionObject, FName SectionName)
+		: FPropertySection(InSectionObject, SectionName) {}
+
+	virtual void GenerateSectionLayout(class ISectionLayoutBuilder& LayoutBuilder) const override
+	{
+		UMovieSceneByteSection* ByteSection = Cast<UMovieSceneByteSection>(&SectionObject);
+		LayoutBuilder.SetSectionAsKeyArea(MakeShareable(new FIntegralKeyArea(ByteSection->GetCurve())));
+	}
+};
+
 FPropertyTrackEditor::FPropertyTrackEditor( TSharedRef<ISequencer> InSequencer )
 	: FMovieSceneTrackEditor( InSequencer ) 
 {
@@ -276,6 +294,7 @@ FPropertyTrackEditor::FPropertyTrackEditor( TSharedRef<ISequencer> InSequencer )
 	ISequencerObjectChangeListener& ObjectChangeListener = InSequencer->GetObjectChangeListener();
 	ObjectChangeListener.GetOnAnimatablePropertyChanged( NAME_FloatProperty ).AddRaw( this, &FPropertyTrackEditor::OnAnimatedPropertyChanged<float, UMovieSceneFloatTrack> );
 	ObjectChangeListener.GetOnAnimatablePropertyChanged( NAME_BoolProperty ).AddRaw( this, &FPropertyTrackEditor::OnAnimatedPropertyChanged<bool, UMovieSceneBoolTrack> );
+	ObjectChangeListener.GetOnAnimatablePropertyChanged( NAME_ByteProperty ).AddRaw( this, &FPropertyTrackEditor::OnAnimatedPropertyChanged<uint8, UMovieSceneByteTrack>);
 	ObjectChangeListener.GetOnAnimatablePropertyChanged( NAME_Vector ).AddRaw( this, &FPropertyTrackEditor::OnAnimatedVectorPropertyChanged );
 	ObjectChangeListener.GetOnAnimatablePropertyChanged( NAME_Vector4 ).AddRaw( this, &FPropertyTrackEditor::OnAnimatedVectorPropertyChanged );
 	ObjectChangeListener.GetOnAnimatablePropertyChanged( NAME_Vector2D ).AddRaw( this, &FPropertyTrackEditor::OnAnimatedVectorPropertyChanged );
@@ -293,6 +312,7 @@ FPropertyTrackEditor::~FPropertyTrackEditor()
 		ISequencerObjectChangeListener& ObjectChangeListener = Sequencer->GetObjectChangeListener();
 		ObjectChangeListener.GetOnAnimatablePropertyChanged(NAME_FloatProperty).RemoveAll(this);
 		ObjectChangeListener.GetOnAnimatablePropertyChanged(NAME_BoolProperty).RemoveAll(this);
+		ObjectChangeListener.GetOnAnimatablePropertyChanged(NAME_ByteProperty).RemoveAll(this);
 		ObjectChangeListener.GetOnAnimatablePropertyChanged(NAME_Vector).RemoveAll(this);
 		ObjectChangeListener.GetOnAnimatablePropertyChanged(NAME_Vector4).RemoveAll(this);
 		ObjectChangeListener.GetOnAnimatablePropertyChanged(NAME_Vector2D).RemoveAll(this);
@@ -314,6 +334,7 @@ bool FPropertyTrackEditor::SupportsType( TSubclassOf<UMovieSceneTrack> Type ) co
 	// We support animatable floats and bools and colors
 	return Type == UMovieSceneFloatTrack::StaticClass() ||
 		Type == UMovieSceneBoolTrack::StaticClass() ||
+		Type == UMovieSceneByteTrack::StaticClass() ||
 		Type == UMovieSceneVectorTrack::StaticClass() ||
 		Type == UMovieSceneColorTrack::StaticClass();
 }
@@ -327,6 +348,7 @@ TSharedRef<ISequencerSection> FPropertyTrackEditor::MakeSectionInterface( UMovie
 	TSharedRef<ISequencerSection> NewSection = TSharedRef<ISequencerSection>(
 		SectionClass == UMovieSceneColorTrack::StaticClass() ? new FColorPropertySection( SectionObject, Track->GetTrackName() ) :
 		SectionClass == UMovieSceneBoolTrack::StaticClass() ? new FBoolPropertySection( SectionObject, Track->GetTrackName() ) :
+		SectionClass == UMovieSceneByteTrack::StaticClass() ? new FBytePropertySection(SectionObject, Track->GetTrackName()) :
 		SectionClass == UMovieSceneVectorTrack::StaticClass() ? new FVectorPropertySection( SectionObject, Track->GetTrackName() ) :
 		SectionClass == UMovieSceneFloatTrack::StaticClass() ? new FFloatPropertySection( SectionObject, Track->GetTrackName() ) :
 		new FPropertySection( SectionObject, Track->GetTrackName() )
diff --git a/Engine/Source/Runtime/MovieSceneCoreTypes/Classes/MovieSceneByteSection.h b/Engine/Source/Runtime/MovieSceneCoreTypes/Classes/MovieSceneByteSection.h
new file mode 100644
--- /dev/null
+++ b/Engine/Source/Runtime/MovieSceneCoreTypes/Classes/MovieSceneByteSection.h
@@ -0,0 +1,54 @@
+// Copyright 1998-2015 Epic Games, Inc. All Rights Reserved.
+
+#pragma once
+
+#include "MovieSceneByteSection.generated.h"
+
+/**
+ * A single byte section
+ */
+UCLASS( MinimalAPI )
+class UMovieSceneByteSection : public UMovieSceneSection
+{
+	GENERATED_UCLASS_BODY()
+public:
+	/**
+	 * Updates this section
+	 *
+	 * @param Position	The position in time within the movie scene
+	 */
+	virtual uint8 Eval( float Position ) const;
+
+	/** 
+	 * Adds a key to the section
+	 *
+	 * @param Time	The location in time where the key should be added
+	 * @param Value	The value of the key
+	 */
+	void AddKey( float Time, uint8 Value );
+	
+	/** 
+	 * Determines if a new key would be new data, or just a duplicate of existing data
+	 *
+	 * @param Time	The location in time where the key would be added
+	 * @param Value	The value of the new key
+	 * @return True if the new key would be new data, false if duplicate
+	 */
+	bool NewKeyIsNewData(float Time, uint8 Value) const;
+
+	/**
+	 * UMovieSceneSection interface 
+	 */
+	virtual void MoveSection( float DeltaPosition ) override;
+	virtual void DilateSection( float DilationFactor, float Origin ) override;
+
+	/** Gets all the keys of this byte section */
+	FIntegralCurve& GetCurve() { return ByteCurve; }
+
+private:
+	/** Ordered curve data */
+	// @todo Sequencer This could be optimized by packing the bytes separately
+	// but that may not be worth the effort
+	UPROPERTY()
+	FIntegralCurve ByteCurve;
+};
diff --git a/Engine/Source/Runtime/MovieSceneCoreTypes/Classes/MovieSceneByteTrack.h b/Engine/Source/Runtime/MovieSceneCoreTypes/Classes/MovieSceneByteTrack.h
new file mode 100644
--- /dev/null
+++ b/Engine/Source/Runtime/MovieSceneCoreTypes/Classes/MovieSceneByteTrack.h
@@ -0,0 +1,41 @@
+// Copyright 1998-2015 Epic Games, Inc. All Rights Reserved.
+
+#pragma once
+
+#include "MovieScene.h"
+#include "MovieSceneTrack.h"
+#include "MovieScenePropertyTrack.h"
+#include "MovieSceneByteTrack.generated.h"
+
+/**
+ * Handles manipulation of float properties in a movie scene
+ */
+UCLASS( MinimalAPI )
+class UMovieSceneByteTrack : public UMovieScenePropertyTrack
+{
+	GENERATED_UCLASS_BODY()
+
+public:
+	/** UMovieSceneTrack interface */
+	virtual UMovieSceneSection* CreateNewSection() override;
+	virtual TSharedPtr<IMovieSceneTrackInstance> CreateInstance() override;
+
+	/**
+	 * Adds a key to a section.  Will create the section if it doesn't exist
+	 *
+	 * @param Time				The time relative to the owning movie scene where the section should be
+	 * @param Value				The value of the key
+	 * @return True if the key was successfully added.
+	 */
+	virtual bool AddKeyToSection( float Time, uint8 Value );
+
+	/**
+	 * Evaluates the track at the playback position
+	 *
+	 * @param Position	The current playback position
+	 * @param LastPosition	The last plackback position
+	 * @param OutByte 	The value at the playback position
+	 * @return true if anything was evaluated. Note: if false is returned OutByte remains unchanged
+	 */
+	virtual bool Eval( float Position, float LastPostion, uint8& OutByte ) const;
+};
diff --git a/Engine/Source/Runtime/MovieSceneCoreTypes/Private/Instances/MovieSceneByteTrackInstance.cpp b/Engine/Source/Runtime/MovieSceneCoreTypes/Private/Instances/MovieSceneByteTrackInstance.cpp
new file mode 100644
--- /dev/null
+++ b/Engine/Source/Runtime/MovieSceneCoreTypes/Private/Instances/MovieSceneByteTrackInstance.cpp
@@ -0,0 +1,30 @@
+// Copyright 1998-2015 Epic Games, Inc. All Rights Reserved.
+
+#include "MovieSceneCoreTypesPCH.h"
+#include "MovieSceneByteTrackInstance.h"
+
+
+FMovieSceneByteTrackInstance::FMovieSceneByteTrackInstance( UMovieSceneByteTrack& InByteTrack )
+{
+	ByteTrack = &InByteTrack;
+	
+	PropertyBindings = MakeShareable( new FTrackInstancePropertyBindings( ByteTrack->GetPropertyName(), ByteTrack->GetPropertyPath() ) );
+}
+
+void FMovieSceneByteTrackInstance::Update( float Position, float LastPosition, const TArray<UObject*>& RuntimeObjects, class IMovieScenePlayer& Player ) 
+{
+	uint8 ByteValue = 0;
+	if( ByteTrack->Eval( Position, LastPosition, ByteValue ) )
+	{
+		for( UObject* Object : RuntimeObjects )
+		{
+			PropertyBindings->CallFunction( Object, &ByteValue );
+		}
+	}
+}
+
+void FMovieSceneByteTrackInstance::RefreshInstance( const TArray<UObject*>& RuntimeObjects, class IMovieScenePlayer& Player )
+{
+	PropertyBindings->UpdateBindings( RuntimeObjects );
+}
+
diff --git a/Engine/Source/Runtime/MovieSceneCoreTypes/Private/Instances/MovieSceneByteTrackInstance.h b/Engine/Source/Runtime/MovieSceneCoreTypes/Private/Instances/MovieSceneByteTrackInstance.h
new file mode 100644
--- /dev/null
+++ b/Engine/Source/Runtime/MovieSceneCoreTypes/Private/Instances/MovieSceneByteTrackInstance.h
@@ -0,0 +1,21 @@
+// Copyright 1998-2015 Epic Games, Inc. All Rights Reserved.
+
+#pragma once
+
+/**
+ * Instance of a UMovieSceneByteTrack
+ */
+class FMovieSceneByteTrackInstance : public IMovieSceneTrackInstance
+{
+public:
+	FMovieSceneByteTrackInstance( UMovieSceneByteTrack& InByteTrack );
+
+	/** IMovieSceneTrackInstance interface */
+	virtual void Update( float Position, float LastPosition, const TArray<UObject*>& RuntimeObjects, class IMovieScenePlayer& Player ) override;
+	virtual void RefreshInstance( const TArray<UObject*>& RuntimeObjects, class IMovieScenePlayer& Player ) override;
+private:
+	/** Byte track that is being instanced */
+	UMovieSceneByteTrack* ByteTrack;
+	/** Runtime property bindings */
+	TSharedPtr<class FTrackInstancePropertyBindings> PropertyBindings;
+};
diff --git a/Engine/Source/Runtime/MovieSceneCoreTypes/Private/MovieSceneByteSection.cpp b/Engine/Source/Runtime/MovieSceneCoreTypes/Private/MovieSceneByteSection.cpp
new file mode 100644
--- /dev/null
+++ b/Engine/Source/Runtime/MovieSceneCoreTypes/Private/MovieSceneByteSection.cpp
@@ -0,0 +1,41 @@
+// Copyright 1998-2015 Epic Games, Inc. All Rights Reserved.
+
+#include "MovieSceneCoreTypesPCH.h"
+#include "MovieSceneByteSection.h"
+
+
+UMovieSceneByteSection::UMovieSceneByteSection( const FObjectInitializer& ObjectInitializer )
+	: Super( ObjectInitializer )
+{
+}
+
+uint8 UMovieSceneByteSection::Eval( float Position ) const
+{
+	return !!ByteCurve.Evaluate(Position);
+}
+
+void UMovieSceneByteSection::MoveSection( float DeltaPosition )
+{
+	Super::MoveSection( DeltaPosition );
+
+	ByteCurve.ShiftCurve(DeltaPosition);
+}
+
+void UMovieSceneByteSection::DilateSection( float DilationFactor, float Origin )
+{
+	Super::DilateSection(DilationFactor, Origin);
+	
+	ByteCurve.ScaleCurve(Origin, DilationFactor);
+
+}
+
+void UMovieSceneByteSection::AddKey( float Time, uint8 Value )
+{
+	Modify();
+	ByteCurve.UpdateOrAddKey(Time, Value ? 1 : 0);
+}
+
+bool UMovieSceneByteSection::NewKeyIsNewData(float Time, uint8 Value) const
+{
+	return ByteCurve.GetNumKeys() == 0 || Eval(Time) != Value;
+}
diff --git a/Engine/Source/Runtime/MovieSceneCoreTypes/Private/MovieSceneByteTrack.cpp b/Engine/Source/Runtime/MovieSceneCoreTypes/Private/MovieSceneByteTrack.cpp
new file mode 100644
--- /dev/null
+++ b/Engine/Source/Runtime/MovieSceneCoreTypes/Private/MovieSceneByteTrack.cpp
@@ -0,0 +1,51 @@
+// Copyright 1998-2015 Epic Games, Inc. All Rights Reserved.
+
+#include "MovieSceneCoreTypesPCH.h"
+#include "MovieSceneByteSection.h"
+#include "MovieSceneByteTrack.h"
+#include "IMovieScenePlayer.h"
+#include "MovieSceneByteTrackInstance.h"
+
+
+UMovieSceneByteTrack::UMovieSceneByteTrack( const FObjectInitializer& ObjectInitializer )
+	: Super( ObjectInitializer )
+{
+}
+
+UMovieSceneSection* UMovieSceneByteTrack::CreateNewSection()
+{
+	return NewObject<UMovieSceneSection>(this, UMovieSceneByteSection::StaticClass(), NAME_None, RF_Transactional);
+}
+
+TSharedPtr<IMovieSceneTrackInstance> UMovieSceneByteTrack::CreateInstance()
+{
+	return MakeShareable( new FMovieSceneByteTrackInstance( *this ) );
+}
+
+bool UMovieSceneByteTrack::AddKeyToSection( float Time, uint8 Value )
+{
+	const UMovieSceneSection* NearestSection = MovieSceneHelpers::FindSectionAtTime( Sections, Time );
+	if (!NearestSection || CastChecked<UMovieSceneByteSection>(NearestSection)->NewKeyIsNewData(Time, Value))
+	{
+		Modify();
+
+		UMovieSceneByteSection* NewSection = CastChecked<UMovieSceneByteSection>(FindOrAddSection(  Time ));
+	
+		NewSection->AddKey( Time, Value );
+
+		return true;
+	}
+	return false;
+}
+
+bool UMovieSceneByteTrack::Eval( float Position, float LastPostion, uint8& OutByte ) const
+{	
+	const UMovieSceneSection* Section = MovieSceneHelpers::FindSectionAtTime( Sections, Position );
+
+	if( Section )
+	{
+		OutByte = CastChecked<UMovieSceneByteSection>( Section )->Eval( Position );
+	}
+
+	return Section != NULL;
+}
