Sequencer - Add features to allow for better in-timeline editing of animations.
+ Add a toolbar button for adding additional objects to animate.
+ Add combo buttons to object binding tracks for adding additional property tracks.
+ Make the animation outliner resizable.
+ Add widgets for editing key values to the animation outliner.
+ Add buttons for navigating to previous and next keys, and for adding a key at the current time to the animation outliner.

diff --git a/Engine/Source/Developer/MovieSceneTools/MovieSceneTools.Build.cs b/Engine/Source/Developer/MovieSceneTools/MovieSceneTools.Build.cs
--- a/Engine/Source/Developer/MovieSceneTools/MovieSceneTools.Build.cs
+++ b/Engine/Source/Developer/MovieSceneTools/MovieSceneTools.Build.cs
@@ -14,7 +14,8 @@ public class MovieSceneTools : ModuleRules
 				"CoreUObject",
 				"Engine",
                 "InputCore",
-				"UnrealEd"
+				"UnrealEd",
+				"Sequencer"
 			}
 		);
 
diff --git a/Engine/Source/Developer/MovieSceneTools/Private/MovieSceneToolHelpers.cpp b/Engine/Source/Developer/MovieSceneTools/Private/MovieSceneToolHelpers.cpp
--- a/Engine/Source/Developer/MovieSceneTools/Private/MovieSceneToolHelpers.cpp
+++ b/Engine/Source/Developer/MovieSceneTools/Private/MovieSceneToolHelpers.cpp
@@ -11,4 +11,66 @@
 #include "ISectionLayoutBuilder.h"
 #include "IKeyArea.h"
 #include "MovieSceneToolHelpers.h"
+#include "SFloatCurveKeyEditor.h"
+#include "SIntegralCurveKeyEditor.h"
+#include "SEnumCurveKeyEditor.h"
+
+void FFloatCurveKeyArea::AddKeyUnique(float Time)
+{
+	FKeyHandle CurrentKeyHandle = Curve->FindKey(Time);
+	if (Curve->IsKeyHandleValid(CurrentKeyHandle) == false)
+	{
+		if (OwningSection->GetStartTime() > Time)
+		{
+			OwningSection->SetStartTime(Time);
+		}
+		if (OwningSection->GetEndTime() < Time)
+		{
+			OwningSection->SetEndTime(Time);
+		}
+		Curve->AddKey(Time, Curve->Eval(Time), false, CurrentKeyHandle);
+	}
+}
+
+TSharedRef<SWidget> FFloatCurveKeyArea::CreateKeyEditor(ISequencer* Sequencer)
+{
+	return SNew(SFloatCurveKeyEditor)
+		.Sequencer(Sequencer)
+		.OwningSection(OwningSection)
+		.Curve(Curve);
+};
+
+void FIntegralKeyArea::AddKeyUnique(float Time)
+{
+	FKeyHandle CurrentKey = Curve.FindKey(Time);
+	if (Curve.IsKeyHandleValid(CurrentKey) == false)
+	{
+		if (OwningSection->GetStartTime() > Time)
+		{
+			OwningSection->SetStartTime(Time);
+		}
+		if (OwningSection->GetEndTime() < Time)
+		{
+			OwningSection->SetEndTime(Time);
+		}
+		Curve.AddKey(Time, Curve.Evaluate(Time), CurrentKey);
+	}
+}
+
+TSharedRef<SWidget> FIntegralKeyArea::CreateKeyEditor(ISequencer* Sequencer)
+{
+	return SNew(SIntegralCurveKeyEditor)
+		.Sequencer(Sequencer)
+		.OwningSection(OwningSection)
+		.Curve(&Curve);
+};
+
+TSharedRef<SWidget> FEnumKeyArea::CreateKeyEditor(ISequencer* Sequencer)
+{
+	return SNew(SEnumCurveKeyEditor)
+		.Sequencer(Sequencer)
+		.OwningSection(OwningSection)
+		.Curve(&Curve)
+		.Enum(Enum);
+};
 
diff --git a/Engine/Source/Developer/MovieSceneTools/Private/PropertyTrackEditor.cpp b/Engine/Source/Developer/MovieSceneTools/Private/PropertyTrackEditor.cpp
--- a/Engine/Source/Developer/MovieSceneTools/Private/PropertyTrackEditor.cpp
+++ b/Engine/Source/Developer/MovieSceneTools/Private/PropertyTrackEditor.cpp
@@ -11,6 +11,7 @@
 #include "MovieSceneColorTrack.h"
 #include "ScopedTransaction.h"
 #include "MovieSceneSection.h"
+#include "ISequencerObjectBindingManager.h"
 #include "ISequencerObjectChangeListener.h"
 #include "ISequencerSection.h"
 #include "ISectionLayoutBuilder.h"
@@ -212,7 +213,7 @@ public:
 	virtual void GenerateSectionLayout( class ISectionLayoutBuilder& LayoutBuilder ) const override
 	{
 		UMovieSceneBoolSection* BoolSection = Cast<UMovieSceneBoolSection>( &SectionObject );
-		LayoutBuilder.SetSectionAsKeyArea( MakeShareable( new FIntegralKeyArea( BoolSection->GetCurve() ) ) );
+		LayoutBuilder.SetSectionAsKeyArea( MakeShareable( new FIntegralKeyArea( BoolSection->GetCurve(), BoolSection ) ) );
 	}
 
 	virtual int32 OnPaintSection( const FGeometry& AllottedGeometry, const FSlateRect& SectionClippingRect, FSlateWindowElementList& OutDrawElements, int32 LayerId, bool bParentEnabled ) const override 
@@ -277,14 +278,25 @@ public:
 class FBytePropertySection : public FPropertySection
 {
 public:
-	FBytePropertySection(UMovieSceneSection& InSectionObject, FName SectionName)
-		: FPropertySection(InSectionObject, SectionName) {}
+	FBytePropertySection(UMovieSceneSection& InSectionObject, FName SectionName, UEnum* Enum)
+		: FPropertySection(InSectionObject, SectionName)
+		, Enum(Enum) {}
 
 	virtual void GenerateSectionLayout(class ISectionLayoutBuilder& LayoutBuilder) const override
 	{
 		UMovieSceneByteSection* ByteSection = Cast<UMovieSceneByteSection>(&SectionObject);
-		LayoutBuilder.SetSectionAsKeyArea(MakeShareable(new FIntegralKeyArea(ByteSection->GetCurve())));
+		if (Enum != nullptr)
+		{
+			LayoutBuilder.SetSectionAsKeyArea(MakeShareable(new FEnumKeyArea(ByteSection->GetCurve(), ByteSection, Enum)));
+		}
+		else
+		{
+			LayoutBuilder.SetSectionAsKeyArea(MakeShareable(new FIntegralKeyArea(ByteSection->GetCurve(), ByteSection)));
+		}
 	}
+
+private:
+	UEnum* Enum;
 };
 
 FPropertyTrackEditor::FPropertyTrackEditor( TSharedRef<ISequencer> InSequencer )
@@ -294,7 +306,7 @@ FPropertyTrackEditor::FPropertyTrackEditor( TSharedRef<ISequencer> InSequencer )
 	ISequencerObjectChangeListener& ObjectChangeListener = InSequencer->GetObjectChangeListener();
 	ObjectChangeListener.GetOnAnimatablePropertyChanged( NAME_FloatProperty ).AddRaw( this, &FPropertyTrackEditor::OnAnimatedPropertyChanged<float, UMovieSceneFloatTrack> );
 	ObjectChangeListener.GetOnAnimatablePropertyChanged( NAME_BoolProperty ).AddRaw( this, &FPropertyTrackEditor::OnAnimatedPropertyChanged<bool, UMovieSceneBoolTrack> );
-	ObjectChangeListener.GetOnAnimatablePropertyChanged( NAME_ByteProperty ).AddRaw( this, &FPropertyTrackEditor::OnAnimatedPropertyChanged<uint8, UMovieSceneByteTrack>);
+	ObjectChangeListener.GetOnAnimatablePropertyChanged( NAME_ByteProperty ).AddRaw( this, &FPropertyTrackEditor::OnAnimatedPropertyChanged<uint8, UMovieSceneByteTrack> );
 	ObjectChangeListener.GetOnAnimatablePropertyChanged( NAME_Vector ).AddRaw( this, &FPropertyTrackEditor::OnAnimatedVectorPropertyChanged );
 	ObjectChangeListener.GetOnAnimatablePropertyChanged( NAME_Vector4 ).AddRaw( this, &FPropertyTrackEditor::OnAnimatedVectorPropertyChanged );
 	ObjectChangeListener.GetOnAnimatablePropertyChanged( NAME_Vector2D ).AddRaw( this, &FPropertyTrackEditor::OnAnimatedVectorPropertyChanged );
@@ -348,7 +360,7 @@ TSharedRef<ISequencerSection> FPropertyTrackEditor::MakeSectionInterface( UMovie
 	TSharedRef<ISequencerSection> NewSection = TSharedRef<ISequencerSection>(
 		SectionClass == UMovieSceneColorTrack::StaticClass() ? new FColorPropertySection( SectionObject, Track->GetTrackName() ) :
 		SectionClass == UMovieSceneBoolTrack::StaticClass() ? new FBoolPropertySection( SectionObject, Track->GetTrackName() ) :
-		SectionClass == UMovieSceneByteTrack::StaticClass() ? new FBytePropertySection(SectionObject, Track->GetTrackName()) :
+		SectionClass == UMovieSceneByteTrack::StaticClass() ? new FBytePropertySection(SectionObject, Track->GetTrackName(), Cast<UMovieSceneByteTrack>(SectionObject.GetOuter())->GetEnum()) :
 		SectionClass == UMovieSceneVectorTrack::StaticClass() ? new FVectorPropertySection( SectionObject, Track->GetTrackName() ) :
 		SectionClass == UMovieSceneFloatTrack::StaticClass() ? new FFloatPropertySection( SectionObject, Track->GetTrackName() ) :
 		new FPropertySection( SectionObject, Track->GetTrackName() )
@@ -369,23 +381,66 @@ void FPropertyTrackEditor::AddKey(const FGuid& ObjectGuid, UObject* AdditionalAs
 	}
 }
 
+UEnum* GetEnumForByteTrack(TSharedPtr<ISequencer> Sequencer, const FGuid& OwnerObjectHandle, FName PropertyName, UMovieSceneByteTrack* ByteTrack)
+{
+	TArray<UObject*> RuntimeObjects;
+	Sequencer->GetObjectBindingManager()->GetRuntimeObjects(Sequencer->GetFocusedMovieSceneInstance(), OwnerObjectHandle, RuntimeObjects);
+
+	TSet<UEnum*> PropertyEnums;
+	for (UObject* RuntimeObject : RuntimeObjects)
+	{
+		UProperty* Property = RuntimeObject->GetClass()->FindPropertyByName(PropertyName);
+		if (Property != nullptr)
+		{
+			UByteProperty* ByteProperty = Cast<UByteProperty>(Property);
+			if (ByteProperty != nullptr && ByteProperty->Enum != nullptr)
+			{
+				PropertyEnums.Add(ByteProperty->Enum);
+			}
+		}
+	}
+
+	UEnum* TrackEnum;
+	if (PropertyEnums.Num() == 1)
+	{
+		TrackEnum = PropertyEnums.Array()[0];
+	}
+	else
+	{
+		TrackEnum = nullptr;
+	}
+	return TrackEnum;
+}
+
+UMovieSceneTrack* FPropertyTrackEditor::AddTrack(UMovieScene* FocusedMovieScene, const FGuid& ObjectHandle, TSubclassOf<class UMovieSceneTrack> TrackClass, FName UniqueTypeName)
+{
+	UMovieSceneTrack* NewTrack = FMovieSceneTrackEditor::AddTrack(FocusedMovieScene, ObjectHandle, TrackClass, UniqueTypeName);
+	if (NewTrack->GetClass() == UMovieSceneByteTrack::StaticClass())
+	{
+		UMovieSceneByteTrack* ByteTrack = Cast<UMovieSceneByteTrack>(NewTrack);
+		UEnum* TrackEnum = GetEnumForByteTrack(GetSequencer(), ObjectHandle, UniqueTypeName, ByteTrack);
+		if (TrackEnum != nullptr)
+		{
+			ByteTrack->SetEnum(TrackEnum);
+		}
+	}
+	return NewTrack;
+}
+
 template <typename Type, typename TrackType>
-void FPropertyTrackEditor::OnAnimatedPropertyChanged(const FKeyPropertyParams& KeyPropertyParams )
+void FPropertyTrackEditor::OnAnimatedPropertyChanged(const FPropertyChangedParams& PropertyChangedParams )
 {
 	// Get the value from the property 
-	Type Value;
-	FPropertyAccess::Result Result = KeyPropertyParams.PropertyHandle->GetValue( Value );
-	check( Result == FPropertyAccess::Success );
-	
-	AnimatablePropertyChanged(TrackType::StaticClass(), KeyPropertyParams.bRequireAutoKey,
-		FOnKeyProperty::CreateRaw(this, &FPropertyTrackEditor::OnKeyProperty<Type, TrackType>, KeyPropertyParams, Value) );
-}
+	Type Value = *PropertyChangedParams.GetPropertyValue<Type>();;
 
+	AnimatablePropertyChanged(TrackType::StaticClass(), PropertyChangedParams.bRequireAutoKey,
+		FOnKeyProperty::CreateRaw(this, &FPropertyTrackEditor::OnKeyProperty<Type, TrackType>, PropertyChangedParams, Value) );
+}
 
-void FPropertyTrackEditor::OnAnimatedVectorPropertyChanged(const FKeyPropertyParams& KeyPropertyParams)
+void FPropertyTrackEditor::OnAnimatedVectorPropertyChanged(const FPropertyChangedParams& PropertyChangedParams)
 {
 	bool bIsVector2D = false, bIsVector = false, bIsVector4 = false;
-	const UStructProperty* StructProp = Cast<const UStructProperty>(KeyPropertyParams.PropertyHandle->GetProperty());
+	const UStructProperty* StructProp = Cast<const UStructProperty>(PropertyChangedParams.PropertyPath.Last());
 	if (StructProp && StructProp->Struct)
 	{
 		FName StructName = StructProp->Struct->GetFName();
@@ -399,58 +454,43 @@ void FPropertyTrackEditor::OnAnimatedVectorPropertyChanged(const FKeyPropertyPar
 		return;
 	}
 
-
-
 	// Get the vector value from the property
-	FPropertyAccess::Result Result = FPropertyAccess::Fail;
-	const bool bRequiresAutoKey = KeyPropertyParams.bRequireAutoKey;
-
-
 	if (bIsVector2D)
 	{
 		FVectorKey<FVector2D> NewKey;
-		NewKey.CurveName = KeyPropertyParams.InnerStructPropertyName;
-		NewKey.bAddKeyEvenIfUnchanged = !bRequiresAutoKey;
+		NewKey.CurveName = PropertyChangedParams.StructPropertyNameToKey;
+		NewKey.bAddKeyEvenIfUnchanged = !PropertyChangedParams.bRequireAutoKey;
+		NewKey.Value = *PropertyChangedParams.GetPropertyValue<FVector2D>();
 
-		Result = KeyPropertyParams.PropertyHandle->GetValue( NewKey.Value );
-
-		FOnKeyProperty OnKeyPropertyDelegate = FOnKeyProperty::CreateRaw( this, &FPropertyTrackEditor::OnKeyProperty<FVectorKey<FVector2D>, UMovieSceneVectorTrack>, KeyPropertyParams, NewKey );
-
-		check( Result == FPropertyAccess::Success );
-		AnimatablePropertyChanged(UMovieSceneVectorTrack::StaticClass(), bRequiresAutoKey, OnKeyPropertyDelegate );
+		FOnKeyProperty OnKeyPropertyDelegate = FOnKeyProperty::CreateRaw( this, &FPropertyTrackEditor::OnKeyProperty<FVectorKey<FVector2D>, UMovieSceneVectorTrack>, PropertyChangedParams, NewKey );
+		AnimatablePropertyChanged(UMovieSceneVectorTrack::StaticClass(), PropertyChangedParams.bRequireAutoKey, OnKeyPropertyDelegate);
 	}
 	else if (bIsVector)
 	{
 		FVectorKey<FVector> NewKey;
-		NewKey.CurveName = KeyPropertyParams.InnerStructPropertyName;
-		NewKey.bAddKeyEvenIfUnchanged = !bRequiresAutoKey;
-
-		Result = KeyPropertyParams.PropertyHandle->GetValue(NewKey.Value);
+		NewKey.CurveName = PropertyChangedParams.StructPropertyNameToKey;
+		NewKey.bAddKeyEvenIfUnchanged = !PropertyChangedParams.bRequireAutoKey;
+		NewKey.Value = *PropertyChangedParams.GetPropertyValue<FVector>();
 
-		FOnKeyProperty OnKeyPropertyDelegate = FOnKeyProperty::CreateRaw(this, &FPropertyTrackEditor::OnKeyProperty<FVectorKey<FVector>, UMovieSceneVectorTrack>, KeyPropertyParams, NewKey );
-
-		check( Result == FPropertyAccess::Success );
-		AnimatablePropertyChanged( UMovieSceneVectorTrack::StaticClass(),bRequiresAutoKey, OnKeyPropertyDelegate );
+		FOnKeyProperty OnKeyPropertyDelegate = FOnKeyProperty::CreateRaw(this, &FPropertyTrackEditor::OnKeyProperty<FVectorKey<FVector>, UMovieSceneVectorTrack>, PropertyChangedParams, NewKey );
+		AnimatablePropertyChanged(UMovieSceneVectorTrack::StaticClass(), PropertyChangedParams.bRequireAutoKey, OnKeyPropertyDelegate);
 	}
 	else if (bIsVector4)
 	{
 		FVectorKey<FVector4> NewKey;
-		NewKey.CurveName = KeyPropertyParams.InnerStructPropertyName;
-		NewKey.bAddKeyEvenIfUnchanged = !bRequiresAutoKey;
-
-		Result = KeyPropertyParams.PropertyHandle->GetValue(NewKey.Value);
-
-		FOnKeyProperty OnKeyPropertyDelegate = FOnKeyProperty::CreateRaw(this, &FPropertyTrackEditor::OnKeyProperty<FVectorKey<FVector4>, UMovieSceneVectorTrack>, KeyPropertyParams, NewKey );
+		NewKey.CurveName = PropertyChangedParams.StructPropertyNameToKey;
+		NewKey.bAddKeyEvenIfUnchanged = !PropertyChangedParams.bRequireAutoKey;
+		NewKey.Value = *PropertyChangedParams.GetPropertyValue<FVector4>();
 
-		check( Result == FPropertyAccess::Success );
-		AnimatablePropertyChanged(UMovieSceneVectorTrack::StaticClass(),bRequiresAutoKey, OnKeyPropertyDelegate );
+		FOnKeyProperty OnKeyPropertyDelegate = FOnKeyProperty::CreateRaw(this, &FPropertyTrackEditor::OnKeyProperty<FVectorKey<FVector4>, UMovieSceneVectorTrack>, PropertyChangedParams, NewKey );
+		AnimatablePropertyChanged(UMovieSceneVectorTrack::StaticClass(), PropertyChangedParams.bRequireAutoKey, OnKeyPropertyDelegate );
 	}
 }
 
-void FPropertyTrackEditor::OnAnimatedColorPropertyChanged(const FKeyPropertyParams& KeyPropertyParams )
+void FPropertyTrackEditor::OnAnimatedColorPropertyChanged(const FPropertyChangedParams& PropertyChangedParams )
 {
 	bool bIsFColor = false, bIsFLinearColor = false, bIsSlateColor = false;
-	const UStructProperty* StructProp = Cast<const UStructProperty>(KeyPropertyParams.PropertyHandle->GetProperty());
+	const UStructProperty* StructProp = Cast<const UStructProperty>(PropertyChangedParams.PropertyPath.Last());
 	if (StructProp && StructProp->Struct)
 	{
 		FName StructName = StructProp->Struct->GetFName();
@@ -465,37 +505,16 @@ void FPropertyTrackEditor::OnAnimatedColorPropertyChanged(const FKeyPropertyPara
 		return;
 	}
 
-
-
-	FName PropertyName = KeyPropertyParams.PropertyHandle->GetProperty()->GetFName();
-					
-	const UClass* PropertyClass = KeyPropertyParams.PropertyHandle->GetPropertyClass();
-
-	TSharedPtr<IPropertyHandle> ChildProperties[4] = {
-		KeyPropertyParams.PropertyHandle->GetChildHandle("R"),
-		KeyPropertyParams.PropertyHandle->GetChildHandle("G"),
-		KeyPropertyParams.PropertyHandle->GetChildHandle("B"),
-		KeyPropertyParams.PropertyHandle->GetChildHandle("A")
-	};
-
-	uint8 FColorChannels[4];
-	float FLinearColorChannels[4];
-	for (int32 i = 0; i < 4; ++i)
-	{
-		FPropertyAccess::Result Result = bIsFColor ?
-			ChildProperties[i]->GetValue(FColorChannels[i]) :
-			ChildProperties[i]->GetValue(FLinearColorChannels[i]);
-		check( Result == FPropertyAccess::Success );
-	}
+	FName PropertyName = PropertyChangedParams.PropertyPath.Last()->GetFName();
 
 	FLinearColor ColorValue;
 	if (bIsFColor)
 	{
-		ColorValue = FColor(FColorChannels[0], FColorChannels[1], FColorChannels[2], FColorChannels[3]).ReinterpretAsLinear();
+		ColorValue = PropertyChangedParams.GetPropertyValue<FColor>()->ReinterpretAsLinear();
 	}
 	else
 	{
-		ColorValue = FLinearColor(FLinearColorChannels[0], FLinearColorChannels[1], FLinearColorChannels[2], FLinearColorChannels[3]);
+		ColorValue = *PropertyChangedParams.GetPropertyValue<FLinearColor>();
 	}
 
 	if( StructProp->HasMetaData("HideAlphaChannel") )
@@ -506,30 +525,30 @@ void FPropertyTrackEditor::OnAnimatedColorPropertyChanged(const FKeyPropertyPara
 
 	FColorKey Key;
 	Key.Value = ColorValue;
-	Key.CurveName = KeyPropertyParams.InnerStructPropertyName;
-	Key.bAddKeyEvenIfUnchanged = !KeyPropertyParams.bRequireAutoKey;
+	Key.CurveName = PropertyChangedParams.StructPropertyNameToKey;
+	Key.bAddKeyEvenIfUnchanged = !PropertyChangedParams.bRequireAutoKey;
 	Key.bIsSlateColor = bIsSlateColor;
 
-	AnimatablePropertyChanged(UMovieSceneColorTrack::StaticClass(), KeyPropertyParams.bRequireAutoKey, FOnKeyProperty::CreateRaw(this, &FPropertyTrackEditor::OnKeyProperty<FColorKey, UMovieSceneColorTrack>, KeyPropertyParams, Key) );
+	AnimatablePropertyChanged(UMovieSceneColorTrack::StaticClass(), PropertyChangedParams.bRequireAutoKey, FOnKeyProperty::CreateRaw(this, &FPropertyTrackEditor::OnKeyProperty<FColorKey, UMovieSceneColorTrack>, PropertyChangedParams, Key) );
 }
 
 
 template <typename Type, typename TrackType>
-void FPropertyTrackEditor::OnKeyProperty( float KeyTime, FKeyPropertyParams KeyPropertyParams, Type Value )
+void FPropertyTrackEditor::OnKeyProperty( float KeyTime, FPropertyChangedParams PropertyChangedParams, Type Value )
 {
-	for( UObject* Object : KeyPropertyParams.ObjectsThatChanged )
+	for( UObject* Object : PropertyChangedParams.ObjectsThatChanged )
 	{
 		FGuid ObjectHandle = FindOrCreateHandleToObject( Object );
 		if (ObjectHandle.IsValid())
 		{
-			FName PropertyName = KeyPropertyParams.PropertyHandle->GetProperty()->GetFName();
+			FName PropertyName = PropertyChangedParams.PropertyPath.Last()->GetFName();
 
 			UMovieSceneTrack* Track = GetTrackForObject( ObjectHandle, TrackType::StaticClass(), PropertyName );
 			if( ensure( Track ) )
 			{
 				TrackType* TypedTrack = CastChecked<TrackType>(Track);
 
-				TypedTrack->SetPropertyNameAndPath( PropertyName, KeyPropertyParams.PropertyPath );
+				TypedTrack->SetPropertyNameAndPath( PropertyName, PropertyChangedParams.GetPropertyPathString() );
 				// Find or add a new section at the auto-key time and changing the property same property
 				// AddKeyToSection is not actually a virtual, it's redefined in each class with a different type
 				bool bSuccessfulAdd = TypedTrack->AddKeyToSection( KeyTime, Value );
diff --git a/Engine/Source/Developer/MovieSceneTools/Private/PropertyTrackEditor.h b/Engine/Source/Developer/MovieSceneTools/Private/PropertyTrackEditor.h
--- a/Engine/Source/Developer/MovieSceneTools/Private/PropertyTrackEditor.h
+++ b/Engine/Source/Developer/MovieSceneTools/Private/PropertyTrackEditor.h
@@ -4,7 +4,7 @@
 
 class IPropertyHandle;
 
-struct FKeyPropertyParams;
+class FPropertyChangedParams;
 
 /**
  * Tools for animatable property types such as floats ands vectors
@@ -32,36 +32,37 @@ public:
 	virtual bool SupportsType( TSubclassOf<UMovieSceneTrack> Type ) const override;
 	virtual TSharedRef<ISequencerSection> MakeSectionInterface( UMovieSceneSection& SectionObject, UMovieSceneTrack* Track ) override;
 	virtual void AddKey(const FGuid& ObjectGuid, UObject* AdditionalAsset = NULL) override;
+	virtual UMovieSceneTrack* AddTrack(UMovieScene* FocusedMovieScene, const FGuid& ObjectHandle, TSubclassOf<class UMovieSceneTrack> TrackClass, FName UniqueTypeName) override;
 
 private:
 	/**
 	 * Called by the details panel when an animatable property changes
 	 *
 	 * @param InObjectsThatChanged	List of objects that changed
-	 * @param PropertyValue			Handle to the property value which changed
+	 * @param KeyPropertyParams		Parameters for the property change.
 	 */
 	template <typename Type, typename TrackType>
-	void OnAnimatedPropertyChanged( const FKeyPropertyParams& KeyPropertyParams );
+	void OnAnimatedPropertyChanged( const FPropertyChangedParams& KeyPropertyParams );
 
 	/**
 	 * Called by the details panel when an animatable vector property changes
 	 *
 	 * @param InObjectsThatChanged	List of objects that changed
-	 * @param PropertyValue			Handle to the property value which changed
+	 * @param KeyPropertyParams		Parameters for the property change.
 	 */
-	void OnAnimatedVectorPropertyChanged( const FKeyPropertyParams& KeyPropertyParams );
+	void OnAnimatedVectorPropertyChanged( const FPropertyChangedParams& KeyPropertyParams );
 	
 	/**
 	 * Called by the details panel when an animatable color property changes
 	 *
 	 * @param InObjectsThatChanged	List of objects that changed
-	 * @param PropertyValue			Handle to the property value which changed
+	 * @param KeyPropertyParams		Parameters for the property change.
 	 */
-	void OnAnimatedColorPropertyChanged( const FKeyPropertyParams& KeyPropertyParams );
+	void OnAnimatedColorPropertyChanged( const FPropertyChangedParams& KeyPropertyParams );
 	
 	/** Delegates for animatable property changed functions */
 	template <typename Type, typename TrackType>
-	void OnKeyProperty( float KeyTime, FKeyPropertyParams KeyPropertyParams, Type Value );
+	void OnKeyProperty( float KeyTime, FPropertyChangedParams KeyPropertyParams, Type Value );
 };
 
 
diff --git a/Engine/Source/Developer/MovieSceneTools/Private/SEnumCurveKeyEditor.cpp b/Engine/Source/Developer/MovieSceneTools/Private/SEnumCurveKeyEditor.cpp
new file mode 100644
--- /dev/null
+++ b/Engine/Source/Developer/MovieSceneTools/Private/SEnumCurveKeyEditor.cpp
@@ -0,0 +1,79 @@
+// Copyright 1998-2015 Epic Games, Inc. All Rights Reserved.
+
+#include "MovieSceneToolsPrivatePCH.h"
+#include "SEnumCurveKeyEditor.h"
+
+#define LOCTEXT_NAMESPACE "EnumCurveKeyEditor"
+
+void SEnumCurveKeyEditor::Construct(const FArguments& InArgs)
+{
+	Sequencer = InArgs._Sequencer;
+	OwningSection = InArgs._OwningSection;
+	Curve = InArgs._Curve;
+	for (int32 i = 0; i < InArgs._Enum->NumEnums(); i++)
+	{
+		FString EnumDisplayValue = InArgs._Enum->GetDisplayNameText(i).ToString();
+		if (EnumDisplayValue.Len() > 0)
+		{
+			EnumDisplayValue = InArgs._Enum->GetEnumName(i);
+		}
+		EnumValues.Add(MakeShareable(new FString(EnumDisplayValue)));
+	}
+	float CurrentTime = Sequencer->GetCurrentLocalTime(*Sequencer->GetFocusedMovieScene());
+	ChildSlot
+	[
+		SNew(SComboBox<TSharedPtr<FString>>)
+		.ButtonStyle(FEditorStyle::Get(), "FlatButton.Light")
+		.OptionsSource(&EnumValues)
+		.InitiallySelectedItem(EnumValues[Curve->Evaluate(CurrentTime)])
+		.OnGenerateWidget(this, &SEnumCurveKeyEditor::OnGenerateWidget)
+		.OnSelectionChanged(this, &SEnumCurveKeyEditor::OnComboSelectionChanged)
+		.ContentPadding(FMargin(2, 0))
+		[
+			SNew(STextBlock)
+			.Font(FEditorStyle::GetFontStyle("Sequencer.AnimationOutliner.RegularFont"))
+			.Text(this, &SEnumCurveKeyEditor::GetCurrentValue)
+		]
+	];
+}
+
+FText SEnumCurveKeyEditor::GetCurrentValue() const
+{
+	float CurrentTime = Sequencer->GetCurrentLocalTime(*Sequencer->GetFocusedMovieScene());
+	return FText::FromString(*EnumValues[Curve->Evaluate(CurrentTime)]);
+}
+
+TSharedRef<SWidget> SEnumCurveKeyEditor::OnGenerateWidget(TSharedPtr<FString> InItem)
+{
+	return SNew(STextBlock)
+		.Text(FText::FromString(*InItem));
+}
+
+void SEnumCurveKeyEditor::OnComboSelectionChanged(TSharedPtr<FString> InSelectedItem, ESelectInfo::Type SelectInfo)
+{
+	FScopedTransaction Transaction(LOCTEXT("SetEnumKey", "Set enum key value"));
+	OwningSection->SetFlags(RF_Transactional);
+	OwningSection->Modify();
+
+	float CurrentTime = Sequencer->GetCurrentLocalTime(*Sequencer->GetFocusedMovieScene());
+
+	bool bKeyWillBeAdded = Curve->IsKeyHandleValid(Curve->FindKey(CurrentTime)) == false;
+	if (bKeyWillBeAdded)
+	{
+		if (OwningSection->GetStartTime() > CurrentTime)
+		{
+			OwningSection->SetStartTime(CurrentTime);
+		}
+		if (OwningSection->GetEndTime() < CurrentTime)
+		{
+			OwningSection->SetEndTime(CurrentTime);
+		}
+	}
+
+	int32 SelectedIndex;
+	EnumValues.Find(InSelectedItem, SelectedIndex);
+	Curve->UpdateOrAddKey(CurrentTime, SelectedIndex);
+	Sequencer->UpdateRuntimeInstances();
+}
+
+#undef LOCTEXT_NAMESPACE
\ No newline at end of file
diff --git a/Engine/Source/Developer/MovieSceneTools/Private/SEnumCurveKeyEditor.h b/Engine/Source/Developer/MovieSceneTools/Private/SEnumCurveKeyEditor.h
new file mode 100644
--- /dev/null
+++ b/Engine/Source/Developer/MovieSceneTools/Private/SEnumCurveKeyEditor.h
@@ -0,0 +1,31 @@
+// Copyright 1998-2015 Epic Games, Inc. All Rights Reserved.
+
+#pragma once
+
+/**
+ * A widget for editing a curve representing enum keys.
+ */
+class SEnumCurveKeyEditor : public SCompoundWidget
+{
+public:
+	SLATE_BEGIN_ARGS(SEnumCurveKeyEditor) {}
+		SLATE_ARGUMENT(ISequencer*, Sequencer)
+		SLATE_ARGUMENT(UMovieSceneSection*, OwningSection)
+		SLATE_ARGUMENT(FIntegralCurve*, Curve)
+		SLATE_ARGUMENT(UEnum*, Enum)
+	SLATE_END_ARGS();
+
+	void Construct(const FArguments& InArgs);
+
+private:
+	FText GetCurrentValue() const;
+
+	TSharedRef<SWidget> OnGenerateWidget(TSharedPtr<FString> InItem);
+
+	void OnComboSelectionChanged(TSharedPtr<FString> InSelectedItem, ESelectInfo::Type SelectInfo);
+
+	ISequencer* Sequencer;
+	UMovieSceneSection* OwningSection;
+	FIntegralCurve* Curve;
+	TArray<TSharedPtr<FString>> EnumValues;
+};
\ No newline at end of file
diff --git a/Engine/Source/Developer/MovieSceneTools/Private/SFloatCurveKeyEditor.cpp b/Engine/Source/Developer/MovieSceneTools/Private/SFloatCurveKeyEditor.cpp
new file mode 100644
--- /dev/null
+++ b/Engine/Source/Developer/MovieSceneTools/Private/SFloatCurveKeyEditor.cpp
@@ -0,0 +1,77 @@
+// Copyright 1998-2015 Epic Games, Inc. All Rights Reserved.
+
+#include "MovieSceneToolsPrivatePCH.h"
+#include "SFloatCurveKeyEditor.h"
+
+#define LOCTEXT_NAMESPACE "FloatCurveKeyEditor"
+
+void SFloatCurveKeyEditor::Construct(const FArguments& InArgs)
+{
+	Sequencer = InArgs._Sequencer;
+	OwningSection = InArgs._OwningSection;
+	Curve = InArgs._Curve;
+	ChildSlot
+	[
+		SNew(SSpinBox<float>)
+		.Style(&FEditorStyle::GetWidgetStyle<FSpinBoxStyle>("Sequencer.AnimationOutliner.KeyEditorSpinBoxStyle"))
+		.Font(FEditorStyle::GetFontStyle("Sequencer.AnimationOutliner.RegularFont"))
+		.MinValue(TOptional<float>())
+		.MaxValue(TOptional<float>())
+		.MaxSliderValue(TOptional<float>())
+		.MinSliderValue(TOptional<float>())
+		.Delta(0.001f)
+		.MinDesiredWidth(60.0f)
+		.Value(this, &SFloatCurveKeyEditor::OnGetKeyValue)
+		.OnValueChanged(this, &SFloatCurveKeyEditor::OnValueChanged)
+		.OnBeginSliderMovement(this, &SFloatCurveKeyEditor::OnBeginSliderMovement)
+		.OnEndSliderMovement(this, &SFloatCurveKeyEditor::OnEndSliderMovement)
+	];
+}
+
+void SFloatCurveKeyEditor::OnBeginSliderMovement()
+{
+	GEditor->BeginTransaction(LOCTEXT("SetFloatKey", "Set float key value"));
+	OwningSection->SetFlags(RF_Transactional);
+	OwningSection->Modify();
+}
+
+void SFloatCurveKeyEditor::OnEndSliderMovement(float Value)
+{
+	if (GEditor->IsTransactionActive())
+	{
+		GEditor->EndTransaction();
+	}
+}
+
+float SFloatCurveKeyEditor::OnGetKeyValue() const
+{
+	float CurrentTime = Sequencer->GetCurrentLocalTime(*Sequencer->GetFocusedMovieScene());
+	return Curve->Eval(CurrentTime);
+}
+
+void SFloatCurveKeyEditor::OnValueChanged(float Value)
+{
+	OwningSection->Modify();
+
+	float CurrentTime = Sequencer->GetCurrentLocalTime(*Sequencer->GetFocusedMovieScene());
+	FKeyHandle CurrentKeyHandle = Curve->FindKey(CurrentTime);
+	if (Curve->IsKeyHandleValid(CurrentKeyHandle))
+	{
+		Curve->SetKeyValue(CurrentKeyHandle, Value);
+	}
+	else
+	{
+		Curve->AddKey(CurrentTime, Value, false, CurrentKeyHandle);
+		if (OwningSection->GetStartTime() > CurrentTime)
+		{
+			OwningSection->SetStartTime(CurrentTime);
+		}
+		if (OwningSection->GetEndTime() < CurrentTime)
+		{
+			OwningSection->SetEndTime(CurrentTime);
+		}
+	}
+	Sequencer->UpdateRuntimeInstances();
+}
+
+#undef LOCTEXT_NAMESPACE
\ No newline at end of file
diff --git a/Engine/Source/Developer/MovieSceneTools/Private/SFloatCurveKeyEditor.h b/Engine/Source/Developer/MovieSceneTools/Private/SFloatCurveKeyEditor.h
new file mode 100644
--- /dev/null
+++ b/Engine/Source/Developer/MovieSceneTools/Private/SFloatCurveKeyEditor.h
@@ -0,0 +1,31 @@
+// Copyright 1998-2015 Epic Games, Inc. All Rights Reserved.
+
+#pragma once
+
+/**
+ * A widget for editing a curve representing float keys.
+ */
+class SFloatCurveKeyEditor : public SCompoundWidget
+{
+public:
+	SLATE_BEGIN_ARGS(SFloatCurveKeyEditor) {}
+		SLATE_ARGUMENT(ISequencer*, Sequencer)
+		SLATE_ARGUMENT(UMovieSceneSection*, OwningSection)
+		SLATE_ARGUMENT(FRichCurve*, Curve)
+	SLATE_END_ARGS();
+
+	void Construct(const FArguments& InArgs);
+
+private:
+	void OnBeginSliderMovement();
+
+	void OnEndSliderMovement(float Value);
+
+	float OnGetKeyValue() const;
+
+	void OnValueChanged(float Value);
+
+	ISequencer* Sequencer;
+	UMovieSceneSection* OwningSection;
+	FRichCurve* Curve;
+};
\ No newline at end of file
diff --git a/Engine/Source/Developer/MovieSceneTools/Private/SIntegralCurveKeyEditor.cpp b/Engine/Source/Developer/MovieSceneTools/Private/SIntegralCurveKeyEditor.cpp
new file mode 100644
--- /dev/null
+++ b/Engine/Source/Developer/MovieSceneTools/Private/SIntegralCurveKeyEditor.cpp
@@ -0,0 +1,75 @@
+// Copyright 1998-2015 Epic Games, Inc. All Rights Reserved.
+
+#include "MovieSceneToolsPrivatePCH.h"
+#include "SIntegralCurveKeyEditor.h"
+
+#define LOCTEXT_NAMESPACE "IntegralCurveKeyEditor"
+
+void SIntegralCurveKeyEditor::Construct(const FArguments& InArgs)
+{
+	Sequencer = InArgs._Sequencer;
+	OwningSection = InArgs._OwningSection;
+	Curve = InArgs._Curve;
+	ChildSlot
+	[
+		SNew(SSpinBox<int32>)
+		.Style(&FEditorStyle::GetWidgetStyle<FSpinBoxStyle>("Sequencer.AnimationOutliner.KeyEditorSpinBoxStyle"))
+		.Font(FEditorStyle::GetFontStyle("Sequencer.AnimationOutliner.RegularFont"))
+		.MinValue(TOptional<int32>())
+		.MaxValue(TOptional<int32>())
+		.MaxSliderValue(TOptional<int32>())
+		.MinSliderValue(TOptional<int32>())
+		.Delta(1)
+		.MinDesiredWidth(60.0f)
+		.Value(this, &SIntegralCurveKeyEditor::OnGetKeyValue)
+		.OnValueChanged(this, &SIntegralCurveKeyEditor::OnValueChanged)
+		.OnBeginSliderMovement(this, &SIntegralCurveKeyEditor::OnBeginSliderMovement)
+		.OnEndSliderMovement(this, &SIntegralCurveKeyEditor::OnEndSliderMovement)
+	];
+}
+
+void SIntegralCurveKeyEditor::OnBeginSliderMovement()
+{
+	GEditor->BeginTransaction(LOCTEXT("SetIntegralKey", "Set integral key value"));
+	OwningSection->SetFlags(RF_Transactional);
+	OwningSection->Modify();
+}
+
+void SIntegralCurveKeyEditor::OnEndSliderMovement(int32 Value)
+{
+	if (GEditor->IsTransactionActive())
+	{
+		GEditor->EndTransaction();
+	}
+}
+
+int32 SIntegralCurveKeyEditor::OnGetKeyValue() const
+{
+	float CurrentTime = Sequencer->GetCurrentLocalTime(*Sequencer->GetFocusedMovieScene());
+	return Curve->Evaluate(CurrentTime);
+}
+
+void SIntegralCurveKeyEditor::OnValueChanged(int32 Value)
+{
+	OwningSection->Modify();
+
+	float CurrentTime = Sequencer->GetCurrentLocalTime(*Sequencer->GetFocusedMovieScene());
+
+	bool bKeyWillBeAdded = Curve->IsKeyHandleValid(Curve->FindKey(CurrentTime)) == false;
+	if (bKeyWillBeAdded)
+	{
+		if (OwningSection->GetStartTime() > CurrentTime)
+		{
+			OwningSection->SetStartTime(CurrentTime);
+		}
+		if (OwningSection->GetEndTime() < CurrentTime)
+		{
+			OwningSection->SetEndTime(CurrentTime);
+		}
+	}
+
+	Curve->UpdateOrAddKey(CurrentTime, Value);
+	Sequencer->UpdateRuntimeInstances();
+}
+
+#undef LOCTEXT_NAMESPACE
\ No newline at end of file
diff --git a/Engine/Source/Developer/MovieSceneTools/Private/SIntegralCurveKeyEditor.h b/Engine/Source/Developer/MovieSceneTools/Private/SIntegralCurveKeyEditor.h
new file mode 100644
--- /dev/null
+++ b/Engine/Source/Developer/MovieSceneTools/Private/SIntegralCurveKeyEditor.h
@@ -0,0 +1,31 @@
+// Copyright 1998-2015 Epic Games, Inc. All Rights Reserved.
+
+#pragma once
+
+/**
+ * A widget for editing a curve representing integer keys.
+ */
+class SIntegralCurveKeyEditor : public SCompoundWidget
+{
+public:
+	SLATE_BEGIN_ARGS(SIntegralCurveKeyEditor) {}
+		SLATE_ARGUMENT(ISequencer*, Sequencer)
+		SLATE_ARGUMENT(UMovieSceneSection*, OwningSection)
+		SLATE_ARGUMENT(FIntegralCurve*, Curve)
+	SLATE_END_ARGS();
+
+	void Construct(const FArguments& InArgs);
+
+private:
+	void OnBeginSliderMovement();
+
+	void OnEndSliderMovement(int32 Value);
+
+	int32 OnGetKeyValue() const;
+
+	void OnValueChanged(int32 Value);
+
+	ISequencer* Sequencer;
+	UMovieSceneSection* OwningSection;
+	FIntegralCurve* Curve;
+};
\ No newline at end of file
diff --git a/Engine/Source/Developer/MovieSceneTools/Public/MovieSceneToolHelpers.h b/Engine/Source/Developer/MovieSceneTools/Public/MovieSceneToolHelpers.h
--- a/Engine/Source/Developer/MovieSceneTools/Public/MovieSceneToolHelpers.h
+++ b/Engine/Source/Developer/MovieSceneTools/Public/MovieSceneToolHelpers.h
@@ -2,24 +2,26 @@
 
 #pragma once
 
+#include "ISequencer.h"
 #include "ISequencerSection.h"
 #include "IKeyArea.h"
 
 /**
  * A key area for float keys
  */
-class FFloatCurveKeyArea : public IKeyArea
+class MOVIESCENETOOLS_API FFloatCurveKeyArea : public IKeyArea
 {
 public:
 	FFloatCurveKeyArea( FRichCurve* InCurve, UMovieSceneSection* InOwningSection )
-		: CurveInfo( InCurve, InOwningSection )
+		: Curve( InCurve )
+		, OwningSection( InOwningSection )
 	{}
 
 	/** IKeyArea interface */
 	virtual TArray<FKeyHandle> GetUnsortedKeyHandles() const override
 	{
 		TArray<FKeyHandle> OutKeyHandles;
-		for (auto It(CurveInfo.Curve->GetKeyHandleIterator()); It; ++It)
+		for (auto It(Curve->GetKeyHandleIterator()); It; ++It)
 		{
 			OutKeyHandles.Add(It.Key());
 		}
@@ -28,35 +30,51 @@ public:
 
 	virtual float GetKeyTime( FKeyHandle KeyHandle ) const override
 	{
-		return CurveInfo.Curve->GetKeyTime( KeyHandle );
+		return Curve->GetKeyTime( KeyHandle );
 	}
 
 	virtual FKeyHandle MoveKey( FKeyHandle KeyHandle, float DeltaPosition ) override
 	{
-		return CurveInfo.Curve->SetKeyTime( KeyHandle, CurveInfo.Curve->GetKeyTime( KeyHandle ) + DeltaPosition );
+		return Curve->SetKeyTime( KeyHandle, Curve->GetKeyTime( KeyHandle ) + DeltaPosition );
 	}
 
 	virtual void DeleteKey(FKeyHandle KeyHandle) override
 	{
-		if ( CurveInfo.Curve->IsKeyHandleValid( KeyHandle ) )
+		if ( Curve->IsKeyHandleValid( KeyHandle ) )
 		{
-			CurveInfo.Curve->DeleteKey( KeyHandle );
+			Curve->DeleteKey( KeyHandle );
 		}
 	}
 
-	virtual FCurveInfo* GetCurveInfo() override { return &CurveInfo; }
+	virtual void AddKeyUnique(float Time) override;
+
+	virtual FRichCurve* GetRichCurve() override { return Curve; }
+
+	virtual UMovieSceneSection* GetOwningSection() override { return OwningSection; }
+
+	virtual bool CanCreateKeyEditor() override
+	{
+		return true;
+	}
+
+	virtual TSharedRef<SWidget> CreateKeyEditor(ISequencer* Sequencer) override;
+
 private:
-	IKeyArea::FCurveInfo CurveInfo;
+	/** The curve which provides the keys for this key area. */
+	FRichCurve* Curve;
+	/** The section that owns this key area. */
+	UMovieSceneSection* OwningSection;
 };
 
 
 
 /** A key area for integral keys */
-class FIntegralKeyArea : public IKeyArea
+class MOVIESCENETOOLS_API FIntegralKeyArea : public IKeyArea
 {
 public:
-	FIntegralKeyArea( FIntegralCurve& InCurve )
+	FIntegralKeyArea(FIntegralCurve& InCurve, UMovieSceneSection* InOwningSection)
 		: Curve(InCurve)
+		, OwningSection(InOwningSection)
 	{}
 
 	/** IKeyArea interface */
@@ -85,10 +103,45 @@ public:
 		Curve.DeleteKey(KeyHandle);
 	}
 
-	virtual IKeyArea::FCurveInfo* GetCurveInfo() override { return nullptr; };
-private:
+	virtual void AddKeyUnique(float Time) override;
+
+	virtual FRichCurve* GetRichCurve() override { return nullptr; };
+
+	virtual UMovieSceneSection* GetOwningSection() override { return OwningSection; }
+
+	virtual bool CanCreateKeyEditor() override
+	{
+		return true;
+	}
+
+	virtual TSharedRef<SWidget> CreateKeyEditor(ISequencer* Sequencer) override;
+
+protected:
 	/** Curve with keys in this area */
 	FIntegralCurve& Curve;
+	/** The section that owns this key area. */
+	UMovieSceneSection* OwningSection;
+};
+
+/** A key area for displaying and editing intragral curves representing enums. */
+class FEnumKeyArea : public FIntegralKeyArea
+{
+public:
+	FEnumKeyArea(FIntegralCurve& InCurve, UMovieSceneSection* InOwningSection, UEnum* InEnum)
+		: FIntegralKeyArea(InCurve, InOwningSection)
+		, Enum(InEnum)
+	{}
+
+	virtual bool CanCreateKeyEditor() override 
+	{
+		return true;
+	}
+
+	virtual TSharedRef<SWidget> CreateKeyEditor(ISequencer* Sequencer) override;
+
+private:
+	/** The enum which provides available integral values for this key area. */
+	UEnum* Enum;
 };
 
 
diff --git a/Engine/Source/Editor/EditorStyle/Private/SlateEditorStyle.cpp b/Engine/Source/Editor/EditorStyle/Private/SlateEditorStyle.cpp
--- a/Engine/Source/Editor/EditorStyle/Private/SlateEditorStyle.cpp
+++ b/Engine/Source/Editor/EditorStyle/Private/SlateEditorStyle.cpp
@@ -427,6 +427,7 @@ void FSlateEditorStyle::FStyle::SetupGeneralStyles()
 				);
 		}
 
+		Set("FontAwesome.7", TTF_FONT("Fonts/FontAwesome", 7));
 		Set("FontAwesome.8", TTF_FONT("Fonts/FontAwesome", 8));
 		Set("FontAwesome.9", TTF_FONT("Fonts/FontAwesome", 9));
 		Set("FontAwesome.10", TTF_FONT("Fonts/FontAwesome", 10));
@@ -1599,8 +1600,16 @@ void FSlateEditorStyle::FStyle::SetupGeneralStyles()
 			.SetPressed( IMAGE_BRUSH("Sequencer/AddKey_Details", FVector2D(11,11), SelectionColor_Pressed ) )
 			.SetNormalPadding(FMargin(0, 1))
 			.SetPressedPadding(FMargin(0, 2, 0, 0));
-
 		Set( "Sequencer.AddKey.Details", DetailsKeyButton );
+
+		const FSplitterStyle OutlinerSplitterStyle = FSplitterStyle()
+			.SetHandleNormalBrush( FSlateNoResource() )
+			.SetHandleHighlightBrush( FSlateNoResource() );
+		Set( "Sequencer.AnimationOutliner.Splitter", OutlinerSplitterStyle );
+
+		const FSpinBoxStyle KeyEditorSpinBoxStyle = FSpinBoxStyle(GetWidgetStyle<FSpinBoxStyle>("SpinBox"))
+			.SetTextPadding(FMargin(2, 0));
+		Set( "Sequencer.AnimationOutliner.KeyEditorSpinBoxStyle", KeyEditorSpinBoxStyle );
 	}
 
 	// Foliage Edit Mode
diff --git a/Engine/Source/Editor/LevelEditor/LevelEditor.Build.cs b/Engine/Source/Editor/LevelEditor/LevelEditor.Build.cs
--- a/Engine/Source/Editor/LevelEditor/LevelEditor.Build.cs
+++ b/Engine/Source/Editor/LevelEditor/LevelEditor.Build.cs
@@ -58,7 +58,8 @@ public class LevelEditor : ModuleRules
 				"WebBrowser",
                 "Persona",
                 "Kismet",
-				"KismetWidgets"
+				"KismetWidgets",
+				"Sequencer"
 			}
 		);
 
diff --git a/Engine/Source/Editor/PropertyEditor/Private/SDetailSingleItemRow.cpp b/Engine/Source/Editor/PropertyEditor/Private/SDetailSingleItemRow.cpp
--- a/Engine/Source/Editor/PropertyEditor/Private/SDetailSingleItemRow.cpp
+++ b/Engine/Source/Editor/PropertyEditor/Private/SDetailSingleItemRow.cpp
@@ -389,7 +389,7 @@ TSharedRef<SWidget> SDetailSingleItemRow::CreateKeyframeButton( FDetailLayoutCus
 		TSharedPtr<IPropertyHandle> Handle = PropertyEditorHelpers::GetPropertyHandle(InCustomization.GetPropertyNode().ToSharedRef(), nullptr, nullptr);
 
 		UClass* ObjectClass = InCustomization.GetPropertyNode()->FindObjectItemParent()->GetObjectBaseClass();
-		SetKeyVisibility = KeyframeHandler.Pin()->IsPropertyKeyable(*ObjectClass, *Handle) ? EVisibility::Visible : EVisibility::Hidden;
+		SetKeyVisibility = KeyframeHandler.Pin()->IsPropertyKeyable(ObjectClass, *Handle) ? EVisibility::Visible : EVisibility::Hidden;
 		
 	}
 
diff --git a/Engine/Source/Editor/PropertyEditor/Public/IDetailKeyframeHandler.h b/Engine/Source/Editor/PropertyEditor/Public/IDetailKeyframeHandler.h
--- a/Engine/Source/Editor/PropertyEditor/Public/IDetailKeyframeHandler.h
+++ b/Engine/Source/Editor/PropertyEditor/Public/IDetailKeyframeHandler.h
@@ -7,7 +7,7 @@ class IDetailKeyframeHandler
 public:
 	virtual ~IDetailKeyframeHandler(){}
 
-	virtual bool IsPropertyKeyable(const UClass& InObjectClass, const class IPropertyHandle& PropertyHandle) const = 0;
+	virtual bool IsPropertyKeyable(UClass* InObjectClass, const class IPropertyHandle& PropertyHandle) const = 0;
 
 	virtual bool IsPropertyKeyingEnabled() const = 0;
 
diff --git a/Engine/Source/Editor/Sequencer/Private/KeyPropertyParams.cpp b/Engine/Source/Editor/Sequencer/Private/KeyPropertyParams.cpp
new file mode 100644
--- /dev/null
+++ b/Engine/Source/Editor/Sequencer/Private/KeyPropertyParams.cpp
@@ -0,0 +1,58 @@
+// Copyright 1998-2015 Epic Games, Inc. All Rights Reserved.
+
+#include "SequencerPrivatePCH.h"
+#include "PropertyHandle.h"
+
+void PropertyHandleToPropertyPath(const UClass* OwnerClass, const IPropertyHandle& InPropertyHandle, TArray<UProperty*>& PropertyPath)
+{
+	PropertyPath.Add(InPropertyHandle.GetProperty());
+	TSharedPtr<IPropertyHandle> CurrentHandle = InPropertyHandle.GetParentHandle();
+	while (CurrentHandle.IsValid() && CurrentHandle->GetProperty() != nullptr)
+	{
+		PropertyPath.Insert(CurrentHandle->GetProperty(), 0);
+		CurrentHandle = CurrentHandle->GetParentHandle();
+	}
+}
+
+FCanKeyPropertyParams::FCanKeyPropertyParams()
+{
+	UClass* ObjectClass = nullptr;
+}
+
+FCanKeyPropertyParams::FCanKeyPropertyParams(UClass* InObjectClass, TArray<UProperty*> InPropertyPath)
+{
+	ObjectClass = InObjectClass;
+	PropertyPath = InPropertyPath;
+}
+
+FCanKeyPropertyParams::FCanKeyPropertyParams(UClass* InObjectClass, const IPropertyHandle& InPropertyHandle)
+{
+	ObjectClass = InObjectClass;
+	PropertyHandleToPropertyPath(ObjectClass, InPropertyHandle, PropertyPath);
+}
+
+FKeyPropertyParams::FKeyPropertyParams()
+{
+}
+
+FKeyPropertyParams::FKeyPropertyParams(TArray<UObject*> InObjectsToKey, TArray<UProperty*> InPropertyPath)
+{
+	ObjectsToKey = InObjectsToKey;
+	PropertyPath = InPropertyPath;
+}
+
+FKeyPropertyParams::FKeyPropertyParams(TArray<UObject*> InObjectsToKey, const IPropertyHandle& InPropertyHandle)
+{
+	ObjectsToKey = InObjectsToKey;
+	PropertyHandleToPropertyPath(InObjectsToKey[0]->GetClass(), InPropertyHandle, PropertyPath);
+}
+
+FString FPropertyChangedParams::GetPropertyPathString() const
+{
+	TArray<FString> PropertyNames;
+	for (UProperty* Property : PropertyPath)
+	{
+		PropertyNames.Add(Property->GetName());
+	}
+	return FString::Join(PropertyNames, TEXT("."));
+}
\ No newline at end of file
diff --git a/Engine/Source/Editor/Sequencer/Private/SAnimationOutlinerView.cpp b/Engine/Source/Editor/Sequencer/Private/SAnimationOutlinerView.cpp
--- a/Engine/Source/Editor/Sequencer/Private/SAnimationOutlinerView.cpp
+++ b/Engine/Source/Editor/Sequencer/Private/SAnimationOutlinerView.cpp
@@ -6,7 +6,10 @@
 #include "ScopedTransaction.h"
 #include "Sequencer.h"
 #include "MovieScene.h"
+#include "MovieSceneSection.h"
 #include "Engine/Selection.h"
+#include "ISequencerObjectBindingManager.h"
+#include "IKeyArea.h"
 
 #define LOCTEXT_NAMESPACE "AnimationOutliner"
 
@@ -51,19 +54,86 @@ void SAnimationOutlinerTreeNode::Construct( const FArguments& InArgs, TSharedRef
 				.Image( this, &SAnimationOutlinerTreeNode::OnGetExpanderImage)
 			]
 		]
+		// Label Slot
 		+ SHorizontalBox::Slot()
-		.AutoWidth()
 		.VAlign( VAlign_Center )
+		.FillWidth(3)
 		[
 			TextWidget
+		]
+		// Editor slot
+		+ SHorizontalBox::Slot()
+		.VAlign(VAlign_Center)
+		.AutoWidth()
+		[
+			// @todo Sequencer - Remove this box and width override.
+			SNew(SBox)
+			.WidthOverride(100)
+			[
+				DisplayNode->GenerateEditWidgetForOutliner()
+			]
+		]
+		// Previous key slot
+		+ SHorizontalBox::Slot()
+		.VAlign(VAlign_Center)
+		.AutoWidth()
+		.Padding(3, 0, 0, 0)
+		[
+			SNew(SButton)
+			.ButtonStyle(FEditorStyle::Get(), "FlatButton")
+			.Visibility(this, &SAnimationOutlinerTreeNode::GetKeyButtonVisibility)
+			.ToolTipText(LOCTEXT("PreviousKeyButton", "Set the time to the previous key"))
+			.OnClicked(this, &SAnimationOutlinerTreeNode::OnPreviousKeyClicked)
+			.ContentPadding(0)
+			[
+				SNew(STextBlock)
+				.TextStyle(FEditorStyle::Get(), "NormalText.Important")
+				.Font(FEditorStyle::Get().GetFontStyle("FontAwesome.7"))
+				.Text(FText::FromString(FString(TEXT("\xf060"))) /*fa-arrow-left*/)
+			]
+		]
+		// Add key slot
+		+ SHorizontalBox::Slot()
+		.VAlign(VAlign_Center)
+		.AutoWidth()
+		[
+			SNew(SButton)
+			.ButtonStyle(FEditorStyle::Get(), "FlatButton")
+			.Visibility(this, &SAnimationOutlinerTreeNode::GetKeyButtonVisibility)
+			.ToolTipText(LOCTEXT("AddKeyButton", "Add a new key at the current time"))
+			.OnClicked(this, &SAnimationOutlinerTreeNode::OnAddKeyClicked)
+			.ContentPadding(0)
+			[
+				SNew(STextBlock)
+				.TextStyle(FEditorStyle::Get(), "NormalText.Important")
+				.Font(FEditorStyle::Get().GetFontStyle("FontAwesome.7"))
+				.Text(FText::FromString(FString(TEXT("\xf055"))) /*fa-plus-circle*/)
+			]
+		]
+		// Next key slot
+		+ SHorizontalBox::Slot()
+		.VAlign(VAlign_Center)
+		.AutoWidth()
+		[
+			SNew(SButton)
+			.ButtonStyle(FEditorStyle::Get(), "FlatButton")
+			.Visibility(this, &SAnimationOutlinerTreeNode::GetKeyButtonVisibility)
+			.ToolTipText(LOCTEXT("NextKeyButton", "Set the time to the next key"))
+			.OnClicked(this, &SAnimationOutlinerTreeNode::OnNextKeyClicked)
+			.ContentPadding(0)
+			[
+				SNew(STextBlock)
+				.TextStyle(FEditorStyle::Get(), "NormalText.Important")
+				.Font(FEditorStyle::Get().GetFontStyle("FontAwesome.7"))
+				.Text(FText::FromString(FString(TEXT("\xf061"))) /*fa-arrow-right*/)
+			]
 		];
 
-
 	ChildSlot
 	[
 		SNew( SBorder )
 		.VAlign( VAlign_Center )
-		.Padding(1.0f)
+		.Padding(0)
 		.BorderImage( this, &SAnimationOutlinerTreeNode::GetNodeBorderImage )
 		[
 			FinalWidget
@@ -73,6 +143,126 @@ void SAnimationOutlinerTreeNode::Construct( const FArguments& InArgs, TSharedRef
 	SetVisibility( TAttribute<EVisibility>( this, &SAnimationOutlinerTreeNode::GetNodeVisibility ) );
 }
 
+void GetAllKeyAreas(TSharedPtr<FSequencerDisplayNode> DisplayNode, TSet<TSharedPtr<IKeyArea>>& KeyAreas)
+{
+	TArray<TSharedPtr<FSequencerDisplayNode>> NodesToCheck;
+	NodesToCheck.Add(DisplayNode);
+	while (NodesToCheck.Num() > 0)
+	{
+		TSharedPtr<FSequencerDisplayNode> NodeToCheck = NodesToCheck[0];
+		NodesToCheck.RemoveAt(0);
+
+		if (NodeToCheck->GetType() == ESequencerNode::Track)
+		{
+			TSharedPtr<FTrackNode> TrackNode = StaticCastSharedPtr<FTrackNode>(NodeToCheck);
+			TArray<TSharedRef<FSectionKeyAreaNode>> KeyAreaNodes;
+			TrackNode->GetChildKeyAreaNodesRecursively(KeyAreaNodes);
+			for (TSharedRef<FSectionKeyAreaNode> KeyAreaNode : KeyAreaNodes)
+			{
+				for (TSharedPtr<IKeyArea> KeyArea : KeyAreaNode->GetAllKeyAreas())
+				{
+					KeyAreas.Add(KeyArea);
+				}
+			}
+		}
+		else
+		{
+			if (NodeToCheck->GetType() == ESequencerNode::KeyArea)
+			{
+				TSharedPtr<FSectionKeyAreaNode> KeyAreaNode = StaticCastSharedPtr<FSectionKeyAreaNode>(NodeToCheck);
+				for (TSharedPtr<IKeyArea> KeyArea : KeyAreaNode->GetAllKeyAreas())
+				{
+					KeyAreas.Add(KeyArea);
+				}
+			}
+			for (TSharedRef<FSequencerDisplayNode> ChildNode : NodeToCheck->GetChildNodes())
+			{
+				NodesToCheck.Add(ChildNode);
+			}
+		}
+	}
+}
+
+FReply SAnimationOutlinerTreeNode::OnPreviousKeyClicked()
+{
+	FSequencer& Sequencer = DisplayNode->GetSequencer();
+	float ClosestPreviousKeyDistance = MAX_FLT;
+	float CurrentTime = Sequencer.GetCurrentLocalTime(*Sequencer.GetFocusedMovieScene());
+	float PreviousTime = 0;
+	bool PreviousKeyFound = false;
+
+	TSet<TSharedPtr<IKeyArea>> KeyAreas;
+	GetAllKeyAreas(DisplayNode, KeyAreas);
+	for (TSharedPtr<IKeyArea> KeyArea : KeyAreas)
+	{
+		for (FKeyHandle& KeyHandle : KeyArea->GetUnsortedKeyHandles())
+		{
+			float KeyTime = KeyArea->GetKeyTime(KeyHandle);
+			if (KeyTime < CurrentTime && CurrentTime - KeyTime < ClosestPreviousKeyDistance)
+			{
+				PreviousTime = KeyTime;
+				ClosestPreviousKeyDistance = CurrentTime - KeyTime;
+				PreviousKeyFound = true;
+			}
+		}
+	}
+
+	if (PreviousKeyFound)
+	{
+		Sequencer.SetGlobalTime(PreviousTime);
+	}
+	return FReply::Handled();
+}
+
+FReply SAnimationOutlinerTreeNode::OnNextKeyClicked()
+{
+	FSequencer& Sequencer = DisplayNode->GetSequencer();
+	float ClosestNextKeyDistance = MAX_FLT;
+	float CurrentTime = Sequencer.GetCurrentLocalTime(*Sequencer.GetFocusedMovieScene());
+	float NextTime = 0;
+	bool NextKeyFound = false;
+
+	TSet<TSharedPtr<IKeyArea>> KeyAreas;
+	GetAllKeyAreas(DisplayNode, KeyAreas);
+	for (TSharedPtr<IKeyArea> KeyArea : KeyAreas)
+	{
+		for (FKeyHandle& KeyHandle : KeyArea->GetUnsortedKeyHandles())
+		{
+			float KeyTime = KeyArea->GetKeyTime(KeyHandle);
+			if (KeyTime > CurrentTime && KeyTime - CurrentTime < ClosestNextKeyDistance)
+			{
+				NextTime = KeyTime;
+				ClosestNextKeyDistance = KeyTime - CurrentTime;
+				NextKeyFound = true;
+			}
+		}
+	}
+
+	if (NextKeyFound)
+	{
+		Sequencer.SetGlobalTime(NextTime);
+	}
+	return FReply::Handled();
+}
+
+FReply SAnimationOutlinerTreeNode::OnAddKeyClicked()
+{
+	FSequencer& Sequencer = DisplayNode->GetSequencer();
+	float CurrentTime = Sequencer.GetCurrentLocalTime(*Sequencer.GetFocusedMovieScene());
+
+	TSet<TSharedPtr<IKeyArea>> KeyAreas;
+	GetAllKeyAreas(DisplayNode, KeyAreas);
+
+	FScopedTransaction Transaction(LOCTEXT("AddKeys", "Add keys at current time"));
+	for (TSharedPtr<IKeyArea> KeyArea : KeyAreas)
+	{
+		UMovieSceneSection* OwningSection = KeyArea->GetOwningSection();
+		OwningSection->SetFlags(RF_Transactional);
+		OwningSection->Modify();
+		KeyArea->AddKeyUnique(CurrentTime);
+	}
+	return FReply::Handled();
+}
 
 FReply SAnimationOutlinerTreeNode::OnMouseButtonDown( const FGeometry& MyGeometry, const FPointerEvent& MouseEvent )
 {
@@ -174,12 +364,17 @@ EVisibility SAnimationOutlinerTreeNode::GetExpanderVisibility() const
 	return DisplayNode->GetNumChildren() > 0 ? EVisibility::Visible : EVisibility::Hidden;
 }
 
+EVisibility SAnimationOutlinerTreeNode::GetKeyButtonVisibility() const
+{
+	return DisplayNode->GetType() == ESequencerNode::Object ? EVisibility::Hidden : EVisibility::Visible;
+}
+
 FText SAnimationOutlinerTreeNode::GetDisplayName() const
 {
 	return DisplayNode->GetDisplayName();
 }
 
-void SAnimationOutlinerView::Construct( const FArguments& InArgs, TSharedRef<FSequencerDisplayNode> InRootNode, TSharedRef<FSequencer> InSequencer )
+void SAnimationOutlinerView::Construct( const FArguments& InArgs, TSharedRef<FSequencerDisplayNode> InRootNode, FSequencer* InSequencer )
 {
 	SAnimationOutlinerViewBase::Construct( SAnimationOutlinerViewBase::FArguments(), InRootNode );
 
@@ -261,9 +456,9 @@ void SAnimationOutlinerView::OnSelectionChanged( TSharedPtr<FSequencerDisplayNod
 
 		// Get the bound objects
 		TArray<UObject*> RuntimeObjects;
-		Sequencer.Pin()->GetRuntimeObjects( Sequencer.Pin()->GetFocusedMovieSceneInstance(), ObjectNode->GetObjectBinding(), RuntimeObjects );
+		Sequencer->GetRuntimeObjects( Sequencer->GetFocusedMovieSceneInstance(), ObjectNode->GetObjectBinding(), RuntimeObjects );
 		
-		if( RuntimeObjects.Num() > 0 && Sequencer.Pin()->IsLevelEditorSequencer() )
+		if( RuntimeObjects.Num() > 0 && Sequencer->IsLevelEditorSequencer() )
 		{
 			const bool bNotifySelectionChanged = false;
 			const bool bDeselectBSP = true;
diff --git a/Engine/Source/Editor/Sequencer/Private/SAnimationOutlinerView.h b/Engine/Source/Editor/Sequencer/Private/SAnimationOutlinerView.h
--- a/Engine/Source/Editor/Sequencer/Private/SAnimationOutlinerView.h
+++ b/Engine/Source/Editor/Sequencer/Private/SAnimationOutlinerView.h
@@ -53,11 +53,25 @@ private:
 	 */
 	EVisibility GetExpanderVisibility() const;
 
+	/**
+	 * @return the visibility of the previous, next, and add key buttons. 
+	 */
+	EVisibility GetKeyButtonVisibility() const;
+
 	/**
 	 * @return The display name for this node.
 	 */
 	FText GetDisplayName() const;
 
+	/** Handles the previous key button being clicked. */
+	FReply OnPreviousKeyClicked();
+
+	/** Handles the next key button being clicked. */
+	FReply OnNextKeyClicked();
+
+	/** Handles the add key button being clicked. */
+	FReply OnAddKeyClicked();
+
 private:
 	/** Layout node the widget is visualizing */
 	TSharedPtr<FSequencerDisplayNode> DisplayNode;
@@ -88,7 +102,7 @@ public:
 	SLATE_END_ARGS()
 
 	/** Construct this widget.  Called by the SNew() Slate macro. */
-	void Construct( const FArguments& InArgs, TSharedRef<FSequencerDisplayNode> RootNode, TSharedRef<FSequencer> InSequencer );
+	void Construct( const FArguments& InArgs, TSharedRef<FSequencerDisplayNode> RootNode, FSequencer* InSequencer );
 
 	/** SAnimationOutlinerView destructor */
 	virtual ~SAnimationOutlinerView();
@@ -111,5 +125,5 @@ private:
 
 private:
 	/** Internal sequencer interface */
-	TWeakPtr<FSequencer> Sequencer;
+	FSequencer* Sequencer;
 };
\ No newline at end of file
diff --git a/Engine/Source/Editor/Sequencer/Private/SSequencer.cpp b/Engine/Source/Editor/Sequencer/Private/SSequencer.cpp
--- a/Engine/Source/Editor/Sequencer/Private/SSequencer.cpp
+++ b/Engine/Source/Editor/Sequencer/Private/SSequencer.cpp
@@ -140,7 +140,6 @@ private:
 	TArray< TRange<float> > CachedFilteredRanges;
 };
 
-
 void SSequencer::Construct( const FArguments& InArgs, TSharedRef< class FSequencer > InSequencer )
 {
 	Sequencer = InSequencer;
@@ -166,9 +165,19 @@ void SSequencer::Construct( const FArguments& InArgs, TSharedRef< class FSequenc
 	bMirrorLabels = true;
 	TSharedRef<ITimeSlider> BottomTimeSlider = SequencerWidgets.CreateTimeSlider( TimeSliderController, bMirrorLabels);
 
+	OnGetAddMenuContent = InArgs._OnGetAddMenuContent;
+
+	OutlinerFillPercentage = .25f;
+	SectionFillPercentage = .75f;
+	OutlinerFillPercentageAttribute.BindRaw(this, &SSequencer::GetOutlinerFillPercentage);
+	SectionFillPercentageAttribute.BindRaw(this, &SSequencer::GetSectionFillPercentage);
+
 	SAssignNew( TrackArea, SSequencerTrackArea, Sequencer.Pin().ToSharedRef() )
 		.ViewRange( InArgs._ViewRange )
-		.OutlinerFillPercent( this, &SSequencer::GetAnimationOutlinerFillPercentage );
+		.OutlinerFillPercentage( OutlinerFillPercentageAttribute )
+		.SectionFillPercentage( SectionFillPercentageAttribute )
+		.OnOutlinerFillPercentageChanged( this, &SSequencer::OnOutlinerFillPercentageChanged )
+		.OnSectionFillPercentageChanged( this, &SSequencer::OnSectionFillPercentageChanged );
 
 	ChildSlot
 	[
@@ -196,12 +205,11 @@ void SSequencer::Construct( const FArguments& InArgs, TSharedRef< class FSequenc
 			+ SVerticalBox::Slot()
 			.AutoHeight()
 			[
-				SNew( SHorizontalBox )
-				+ SHorizontalBox::Slot()
-				// @todo Sequencer Do not change the paddings or the sliders scrub widgets wont line up
-				.Padding( FMargin(0) )
-				.FillWidth( TAttribute<float>( this, &SSequencer::GetAnimationOutlinerFillPercentage ) )
-				.VAlign(VAlign_Center)
+				SNew( SSplitter )
+				.Style( FEditorStyle::Get(), "Sequencer.AnimationOutliner.Splitter" )
+				+ SSplitter::Slot()
+				.Value(OutlinerFillPercentageAttribute)
+				.OnSlotResized(SSplitter::FOnSlotResized::CreateSP(this, &SSequencer::OnOutlinerFillPercentageChanged))
 				[
 					SNew( SBox )
 					.Padding( FMargin( 0,0,4,0) )
@@ -211,8 +219,9 @@ void SSequencer::Construct( const FArguments& InArgs, TSharedRef< class FSequenc
 						.OnTextChanged( this, &SSequencer::OnOutlinerSearchChanged )
 					]
 				]
-				+ SHorizontalBox::Slot()
-				.FillWidth(1.0f)
+				+ SSplitter::Slot()
+				.Value(SectionFillPercentageAttribute)
+				.OnSlotResized(SSplitter::FOnSlotResized::CreateSP(this, &SSequencer::OnSectionFillPercentageChanged))
 				[
 					SNew( SBorder )
 					// @todo Sequencer Do not change the paddings or the sliders scrub widgets wont line up
@@ -238,16 +247,19 @@ void SSequencer::Construct( const FArguments& InArgs, TSharedRef< class FSequenc
 				]
 				+ SOverlay::Slot()
 				[
-					SNew( SHorizontalBox )
+					SNew( SSplitter )
+					.Style(FEditorStyle::Get(), "Sequencer.AnimationOutliner.Splitter")
 					.Visibility( EVisibility::HitTestInvisible )
-					+ SHorizontalBox::Slot()
-					.FillWidth( TAttribute<float>( this, &SSequencer::GetAnimationOutlinerFillPercentage ) )
+					+ SSplitter::Slot()
+					.Value(OutlinerFillPercentageAttribute)
+					.OnSlotResized(SSplitter::FOnSlotResized::CreateSP(this, &SSequencer::OnOutlinerFillPercentageChanged))
 					[
 						// Take up space but display nothing. This is required so that all areas dependent on time align correctly
 						SNullWidget::NullWidget
 					]
-					+ SHorizontalBox::Slot()
-					.FillWidth( 1.0f )
+					+ SSplitter::Slot()
+					.Value(SectionFillPercentageAttribute)
+					.OnSlotResized(SSplitter::FOnSlotResized::CreateSP(this, &SSequencer::OnSectionFillPercentageChanged))
 					[
 						SNew(SShotFilterOverlay, Sequencer)
 						.ViewRange( InArgs._ViewRange )
@@ -261,15 +273,21 @@ void SSequencer::Construct( const FArguments& InArgs, TSharedRef< class FSequenc
 			+ SVerticalBox::Slot()
 			.AutoHeight()
 			[
-				SNew( SHorizontalBox )
-				+ SHorizontalBox::Slot()
-				.FillWidth( TAttribute<float>( this, &SSequencer::GetAnimationOutlinerFillPercentage ) )
-				.HAlign(HAlign_Center)
+				SNew( SSplitter )
+				.Style(FEditorStyle::Get(), "Sequencer.AnimationOutliner.Splitter")
+				+ SSplitter::Slot()
+				.Value(OutlinerFillPercentageAttribute)
+				.OnSlotResized(SSplitter::FOnSlotResized::CreateSP(this, &SSequencer::OnOutlinerFillPercentageChanged))
 				[
-					MakeTransportControls()
+					SNew(SBox)
+					.HAlign(HAlign_Center)
+					[
+						MakeTransportControls()
+					]
 				]
-				+ SHorizontalBox::Slot()
-				.FillWidth(1.0f)
+				+ SSplitter::Slot()
+				.Value(SectionFillPercentageAttribute)
+				.OnSlotResized(SSplitter::FOnSlotResized::CreateSP(this, &SSequencer::OnSectionFillPercentageChanged))
 				[
 					SNew( SBorder )
 					// @todo Sequencer Do not change the paddings or the sliders scrub widgets wont line up
@@ -284,14 +302,17 @@ void SSequencer::Construct( const FArguments& InArgs, TSharedRef< class FSequenc
 			+ SVerticalBox::Slot()
 			.AutoHeight()
 			[
-				SNew( SHorizontalBox )
-				+ SHorizontalBox::Slot()
-				.FillWidth( TAttribute<float>( this, &SSequencer::GetAnimationOutlinerFillPercentage ) )
+				SNew( SSplitter )
+				.Style(FEditorStyle::Get(), "Sequencer.AnimationOutliner.Splitter")
+				+ SSplitter::Slot()
+				.Value(OutlinerFillPercentageAttribute)
+				.OnSlotResized(SSplitter::FOnSlotResized::CreateSP(this, &SSequencer::OnOutlinerFillPercentageChanged))
 				[
 					SNew( SSpacer )
 				]
-				+ SHorizontalBox::Slot()
-				.FillWidth( 1.0f )
+				+ SSplitter::Slot()
+				.Value(SectionFillPercentageAttribute)
+				.OnSlotResized(SSplitter::FOnSlotResized::CreateSP(this, &SSequencer::OnSectionFillPercentageChanged))
 				[
 					SAssignNew( BreadcrumbTrail, SBreadcrumbTrail<FSequencerBreadcrumb> )
 					.Visibility(this, &SSequencer::GetBreadcrumbTrailVisibility)
@@ -311,6 +332,51 @@ TSharedRef<SWidget> SSequencer::MakeToolBar()
 	ToolBarBuilder.BeginSection("Base Commands");
 	{
 		ToolBarBuilder.SetLabelVisibility(EVisibility::Visible);
+
+		if (OnGetAddMenuContent.IsBound())
+		{
+			ToolBarBuilder.AddWidget(
+				SNew(SComboButton)
+				.OnGetMenuContent(this, &SSequencer::MakeAddMenu)
+				.ButtonStyle(FEditorStyle::Get(), "FlatButton.Success")
+				.ContentPadding(FMargin(2.0f, 1.0f))
+				.HasDownArrow(false)
+				.ButtonContent()
+				[
+					SNew(SHorizontalBox)
+
+					+ SHorizontalBox::Slot()
+					.VAlign(VAlign_Center)
+					.AutoWidth()
+					[
+						SNew(STextBlock)
+						.TextStyle(FEditorStyle::Get(), "NormalText.Important")
+						.Font(FEditorStyle::Get().GetFontStyle("FontAwesome.10"))
+						.Text(FText::FromString(FString(TEXT("\xf067"))) /*fa-plus*/)
+					]
+
+					+ SHorizontalBox::Slot()
+					.AutoWidth()
+					.Padding(4, 0, 0, 0)
+					[
+						SNew(STextBlock)
+						.TextStyle(FEditorStyle::Get(), "NormalText.Important")
+						.Text(LOCTEXT("AddButton", "Track"))
+					]
+
+					+ SHorizontalBox::Slot()
+					.VAlign(VAlign_Center)
+					.AutoWidth()
+					.Padding(4, 0, 0, 0)
+					[
+						SNew(STextBlock)
+						.TextStyle(FEditorStyle::Get(), "NormalText.Important")
+						.Font(FEditorStyle::Get().GetFontStyle("FontAwesome.10"))
+						.Text(FText::FromString(FString(TEXT("\xf0d7"))) /*fa-caret-down*/)
+					]
+				]);
+		}
+
 		ToolBarBuilder.AddToolBarButton( FSequencerCommands::Get().ToggleAutoKeyEnabled );
 
 		if ( Sequencer.Pin()->IsLevelEditorSequencer() )
@@ -381,6 +447,11 @@ TSharedRef<SWidget> SSequencer::MakeToolBar()
 	return ToolBarBuilder.MakeWidget();
 }
 
+TSharedRef<SWidget> SSequencer::MakeAddMenu()
+{
+	return OnGetAddMenuContent.Execute(Sequencer.Pin().ToSharedRef());
+}
+
 TSharedRef<SWidget> SSequencer::MakeSnapMenu()
 {
 	FMenuBuilder MenuBuilder( false, Sequencer.Pin()->GetCommandBindings() );
@@ -538,16 +609,19 @@ void SSequencer::OnOutlinerSearchChanged( const FText& Filter )
 TSharedRef<SWidget> SSequencer::MakeSectionOverlay( TSharedRef<FSequencerTimeSliderController> TimeSliderController, const TAttribute< TRange<float> >& ViewRange, const TAttribute<float>& ScrubPosition, bool bTopOverlay )
 {
 	return
-		SNew( SHorizontalBox )
+		SNew( SSplitter )
+		.Style(FEditorStyle::Get(), "Sequencer.AnimationOutliner.Splitter")
 		.Visibility( EVisibility::HitTestInvisible )
-		+ SHorizontalBox::Slot()
-		.FillWidth( TAttribute<float>( this, &SSequencer::GetAnimationOutlinerFillPercentage ) )
+		+ SSplitter::Slot()
+		.Value(OutlinerFillPercentageAttribute)
+		.OnSlotResized(SSplitter::FOnSlotResized::CreateSP(this, &SSequencer::OnOutlinerFillPercentageChanged))
 		[
 			// Take up space but display nothing. This is required so that all areas dependent on time align correctly
 			SNullWidget::NullWidget
 		]
-		+ SHorizontalBox::Slot()
-		.FillWidth( 1.0f )
+		+ SSplitter::Slot()
+		.Value(SectionFillPercentageAttribute)
+		.OnSlotResized(SSplitter::FOnSlotResized::CreateSP(this, &SSequencer::OnSectionFillPercentageChanged))
 		[
 			SNew( SSequencerSectionOverlay, TimeSliderController )
 			.DisplayScrubPosition( bTopOverlay )
@@ -890,4 +964,14 @@ EVisibility SSequencer::GetCurveEditorToolBarVisibility() const
 	return GetDefault<USequencerSettings>()->GetShowCurveEditor() ? EVisibility::Visible : EVisibility::Collapsed;
 }
 
+void SSequencer::OnOutlinerFillPercentageChanged(float InOutlinerFillPercentage)
+{
+	OutlinerFillPercentage = InOutlinerFillPercentage;
+}
+
+void SSequencer::OnSectionFillPercentageChanged(float InSectionFillPercentage)
+{
+	SectionFillPercentage = InSectionFillPercentage;
+}
+
 #undef LOCTEXT_NAMESPACE
diff --git a/Engine/Source/Editor/Sequencer/Private/SSequencer.h b/Engine/Source/Editor/Sequencer/Private/SSequencer.h
--- a/Engine/Source/Editor/Sequencer/Private/SSequencer.h
+++ b/Engine/Source/Editor/Sequencer/Private/SSequencer.h
@@ -9,7 +9,7 @@ namespace SequencerLayoutConstants
 	/** Padding between each node */
 	const float NodePadding = 3.0f;
 	/** Height of each object node */
-	const float ObjectNodeHeight = 25.0f;
+	const float ObjectNodeHeight = 20.0f;
 	/** Height of each section area if there are no sections (note: section areas may be larger than this if they have children. This is the height of a section area with no children or all children hidden) */
 	const float SectionAreaDefaultHeight = 15.0f;
 	/** Height of each key area */
@@ -61,6 +61,8 @@ public:
 		SLATE_EVENT( FOnViewRangeChanged, OnViewRangeChanged )
 		/** Called when the user changes the scrub position */
 		SLATE_EVENT( FOnScrubPositionChanged, OnScrubPositionChanged )
+		/** Called to populate the add combo button in the toolbar. */
+		SLATE_EVENT( FOnGetAddMenuContent, OnGetAddMenuContent )
 	SLATE_END_ARGS()
 
 
@@ -91,9 +93,12 @@ private:
 	/** Empty active timer to ensure Slate ticks during Sequencer playback */
 	EActiveTimerReturnType EnsureSlateTickDuringPlayback(double InCurrentTime, float InDeltaTime);	
 
-	/** Makes the toolbar for the outline section. */
+	/** Makes the toolbar. */
 	TSharedRef<SWidget> MakeToolBar();
 
+	/** Makes the add menu for the toolbar. */
+	TSharedRef<SWidget> MakeAddMenu();
+
 	/** Makes the snapping menu for the toolbar. */
 	TSharedRef<SWidget> MakeSnapMenu();
 
@@ -134,9 +139,13 @@ private:
 
 	/**
 	 * @return The fill percentage of the animation outliner
-	 * @todo Sequencer Make this user adjustable
 	 */
-	float GetAnimationOutlinerFillPercentage() const { return .25f; }
+	float GetOutlinerFillPercentage() const { return OutlinerFillPercentage; }
+
+	/**
+	 * @return The fill percentage of the section area.
+	 */
+	float GetSectionFillPercentage() const { return SectionFillPercentage; }
 
 	/**
 	 * Creates an overlay for various components of the time slider that should be displayed on top or bottom of the section area
@@ -199,6 +208,12 @@ private:
 	/** Gets whether or not the curve editor toolbar should be visibe. */
 	EVisibility GetCurveEditorToolBarVisibility() const;
 
+	/** Called when the outliner fill percentage is changed by a splitter slot. */
+	void OnOutlinerFillPercentageChanged(float OutlinerFillPercentage);
+
+	/** Called when the section fill percentage is changed by a splitter slot. */
+	void OnSectionFillPercentageChanged(float SectionFillPercentage);
+
 private:
 	/** Section area widget */
 	TSharedPtr<class SSequencerTrackArea> TrackArea;
@@ -210,7 +225,18 @@ private:
 	TAttribute<bool> bShowBreadcrumbTrail;
 	/** The main sequencer interface */
 	TWeakPtr<FSequencer> Sequencer;
+	
+	/** The fill percentage of the outliner area. */
+	float OutlinerFillPercentage;
+	/** The fill percentage of the section area. */
+	float SectionFillPercentage;
+	/** An attribute for binding the fill percentage of the outliner area to a splitter slot. */
+	TAttribute<float>OutlinerFillPercentageAttribute;
+	/** An attribute for binding the fill percentage of the section area to a splitter slot. */
+	TAttribute<float>SectionFillPercentageAttribute;
 
 	/** Whether the active timer is currently registered */
 	bool bIsActiveTimerRegistered;
+
+	FOnGetAddMenuContent OnGetAddMenuContent;
 };
diff --git a/Engine/Source/Editor/Sequencer/Private/SSequencerTrackArea.cpp b/Engine/Source/Editor/Sequencer/Private/SSequencerTrackArea.cpp
--- a/Engine/Source/Editor/Sequencer/Private/SSequencerTrackArea.cpp
+++ b/Engine/Source/Editor/Sequencer/Private/SSequencerTrackArea.cpp
@@ -11,9 +11,15 @@
 void SSequencerTrackArea::Construct( const FArguments& InArgs, TSharedRef<FSequencer> InSequencer )
 {
 	ViewRange = InArgs._ViewRange;
-	OutlinerFillPercent = InArgs._OutlinerFillPercent;
+	OutlinerFillPercentage = InArgs._OutlinerFillPercentage;
+	SectionFillPercentage = InArgs._SectionFillPercentage;
+	OnOutlinerFillPercentageChanged = InArgs._OnOutlinerFillPercentageChanged;
+	OnSectionFillPercentageChanged = InArgs._OnSectionFillPercentageChanged;
 	Sequencer = InSequencer;
 
+	ScrollbarFillPercentage.BindRaw(this, &SSequencerTrackArea::GetScrollBarFillPercentage);
+	ScrollbarSpacerFillPercentage.BindRaw(this, &SSequencerTrackArea::GetScrollBarSpacerFillPercentage);
+
 	TSharedRef<SScrollBar> ScrollBar =
 		SNew(SScrollBar)
 		.Thickness(FVector2D(5.0f, 5.0f));
@@ -28,14 +34,18 @@ void SSequencerTrackArea::Construct( const FArguments& InArgs, TSharedRef<FSeque
 		]
 		+ SOverlay::Slot()
 		[
-			SNew(SHorizontalBox)
-			+ SHorizontalBox::Slot()
-			.FillWidth(OutlinerFillPercent)
+			SNew(SSplitter)
+			.Style(FEditorStyle::Get(), "Sequencer.AnimationOutliner.Splitter")
+			.Visibility(EVisibility::SelfHitTestInvisible)
+			+ SSplitter::Slot()
+			.Value(OutlinerFillPercentage)
+			.OnSlotResized(SSplitter::FOnSlotResized::CreateSP(this, &SSequencerTrackArea::OutlinerFillPercentageChanged))
 			[
 				SNew(SSpacer)
 			]
-			+ SHorizontalBox::Slot()
-			.FillWidth(1.0f)
+			+ SSplitter::Slot()
+			.Value(SectionFillPercentage)
+			.OnSlotResized(SSplitter::FOnSlotResized::CreateSP(this, &SSequencerTrackArea::SectionFillPercentageChanged))
 			[
 				SAssignNew(CurveEditor, SSequencerCurveEditor, InSequencer)
 				.Visibility( this, &SSequencerTrackArea::GetCurveEditorVisibility )
@@ -44,15 +54,20 @@ void SSequencerTrackArea::Construct( const FArguments& InArgs, TSharedRef<FSeque
 		]
 		+ SOverlay::Slot()
 		[
-			SNew(SHorizontalBox)
-			+ SHorizontalBox::Slot()
-			.HAlign(HAlign_Right)
-			.FillWidth( TAttribute<float>( this, &SSequencerTrackArea::GetScrollBarSlotFill ) )
+			SNew(SSplitter)
+			.Style(FEditorStyle::Get(), "Sequencer.AnimationOutliner.Splitter")
+			.Visibility(EVisibility::SelfHitTestInvisible)
+			+ SSplitter::Slot()
+			.Value(ScrollbarFillPercentage)
 			[
-				ScrollBar
+				SNew(SBox)
+				.HAlign(HAlign_Right)
+				[
+					ScrollBar
+				]
 			]
-			+ SHorizontalBox::Slot()
-			.FillWidth( TAttribute<float>( this, &SSequencerTrackArea::GetScrollBarSpacerSlotFill ) )
+			+ SSplitter::Slot()
+			.Value(ScrollbarSpacerFillPercentage)
 			[
 				SNew( SSpacer )
 			]
@@ -103,17 +118,22 @@ void SSequencerTrackArea::GenerateWidgetForNode( TSharedRef<FSequencerDisplayNod
 {
 	ScrollBox->AddSlot()
 	[
-		SNew( SHorizontalBox )
-		+ SHorizontalBox::Slot()
-		.FillWidth( OutlinerFillPercent )
-		.Padding( FMargin(0.0f, 2.0f ) )
+		SNew( SSplitter )
+		.Style(FEditorStyle::Get(), "Sequencer.AnimationOutliner.Splitter")
+		+ SSplitter::Slot()
+		.Value(OutlinerFillPercentage)
+		.OnSlotResized(SSplitter::FOnSlotResized::CreateSP(this, &SSequencerTrackArea::OutlinerFillPercentageChanged))
 		[
-			// Generate a widget for the animation outliner portion of the node
-			Node->GenerateWidgetForOutliner( Sequencer.Pin().ToSharedRef() )
+			SNew(SBox)
+			.Padding(this, &SSequencerTrackArea::GetOutlinerPadding)
+			[
+				// Generate a widget for the animation outliner portion of the node
+				Node->GenerateContainerWidgetForOutliner()
+			]
 		]
-		+ SHorizontalBox::Slot()
-		.FillWidth( 1.0f )
-		.Padding( FMargin(0.0f, 2.0f) )
+		+ SSplitter::Slot()
+		.Value(SectionFillPercentage)
+		.OnSlotResized(SSplitter::FOnSlotResized::CreateSP(this, &SSequencerTrackArea::SectionFillPercentageChanged))
 		[
 			SNew(SBox)
 			.Visibility(this, &SSequencerTrackArea::GetSectionControlVisibility)
@@ -135,12 +155,27 @@ EVisibility SSequencerTrackArea::GetSectionControlVisibility() const
 	return GetDefault<USequencerSettings>()->GetShowCurveEditor() ? EVisibility::Hidden : EVisibility::Visible;
 }
 
-float SSequencerTrackArea::GetScrollBarSlotFill() const
+float SSequencerTrackArea::GetScrollBarFillPercentage() const
+{
+	return GetDefault<USequencerSettings>()->GetShowCurveEditor() ? OutlinerFillPercentage.Get() : 1.0f;
+}
+
+float SSequencerTrackArea::GetScrollBarSpacerFillPercentage() const
+{
+	return GetDefault<USequencerSettings>()->GetShowCurveEditor() ? SectionFillPercentage.Get() : 0.0f;
+}
+
+FMargin SSequencerTrackArea::GetOutlinerPadding() const
+{
+	return GetDefault<USequencerSettings>()->GetShowCurveEditor() ? FMargin(0, 0, 10, SequencerLayoutConstants::NodePadding) : FMargin(0, 0, 0, SequencerLayoutConstants::NodePadding);
+}
+
+void SSequencerTrackArea::OutlinerFillPercentageChanged(float Value)
 {
-	return GetDefault<USequencerSettings>()->GetShowCurveEditor() ? OutlinerFillPercent.Get() : 1.0f;
+	OnOutlinerFillPercentageChanged.ExecuteIfBound(Value);
 }
 
-float SSequencerTrackArea::GetScrollBarSpacerSlotFill() const
+void SSequencerTrackArea::SectionFillPercentageChanged(float Value)
 {
-	return GetDefault<USequencerSettings>()->GetShowCurveEditor() ? 1.0f : 0.0f;
+	OnSectionFillPercentageChanged.ExecuteIfBound(Value);
 }
diff --git a/Engine/Source/Editor/Sequencer/Private/SSequencerTrackArea.h b/Engine/Source/Editor/Sequencer/Private/SSequencerTrackArea.h
--- a/Engine/Source/Editor/Sequencer/Private/SSequencerTrackArea.h
+++ b/Engine/Source/Editor/Sequencer/Private/SSequencerTrackArea.h
@@ -8,12 +8,20 @@
 class SSequencerTrackArea : public SCompoundWidget
 {
 public:
+	DECLARE_DELEGATE_OneParam(FOnValueChanged, float)
+
 	SLATE_BEGIN_ARGS( SSequencerTrackArea )
 	{}
 		/** The range of time being viewed */
 		SLATE_ATTRIBUTE( TRange<float>, ViewRange )
 		/** Percentage of horizontal space that the animation outliner takes up */
-		SLATE_ATTRIBUTE( float, OutlinerFillPercent )
+		SLATE_ATTRIBUTE( float, OutlinerFillPercentage )
+		/** Percentage of horizontal space that the section area takes up. */
+		SLATE_ATTRIBUTE( float, SectionFillPercentage )
+		/** Called whenever the outliner fill percentage is changed by a splitter. */
+		SLATE_EVENT( FOnValueChanged, OnOutlinerFillPercentageChanged )
+		/** Called whenever the section fill percentage is changed by a splitter. */
+		SLATE_EVENT( FOnValueChanged, OnSectionFillPercentageChanged )
 	SLATE_END_ARGS()
 
 	void Construct( const FArguments& InArgs, TSharedRef<FSequencer> InSequencer );
@@ -47,10 +55,17 @@ private:
 
 	/** Gets the fill value for the slot which contains the scrollbar.  This allows the scroll bar to move depending
 	  * whether or not the curve editor is visible. */
-	float GetScrollBarSlotFill() const;
+	float GetScrollBarFillPercentage() const;
 	/** Gets the fill value for the spacer slot which is to the right of the scrollbar.  This allows the scroll bar 
 	  * to move depending whether or not the curve editor is visible. */
-	float GetScrollBarSpacerSlotFill() const;
+	float GetScrollBarSpacerFillPercentage() const;
+
+	FMargin GetOutlinerPadding() const;
+
+	/** Called whenever the outliner fill percentage is changed by a splitter. */
+	void OutlinerFillPercentageChanged(float Value);
+	/** Called whenever the section fill percentage is changed by a splitter. */
+	void SectionFillPercentageChanged(float Value);
 
 private:
 	/** Scrollable area to display widgets */
@@ -58,7 +73,17 @@ private:
 	/** The current view range */
 	TAttribute< TRange<float> > ViewRange;
 	/** The fill percentage of the animation outliner */
-	TAttribute<float> OutlinerFillPercent;
+	TAttribute<float> OutlinerFillPercentage;
+	/** The fill percentage of the section area */
+	TAttribute<float> SectionFillPercentage;
+	/** Delegate to be called whenever the outliner fill percentage is changed */
+	FOnValueChanged OnOutlinerFillPercentageChanged;
+	/** Delegate to be called whenever the section fill percentage is changed */
+	FOnValueChanged OnSectionFillPercentageChanged;
+	/** Attribute representing the scrollbar fill percentage */
+	TAttribute<float> ScrollbarFillPercentage;
+	/** Attribute representing the spacer fill percentage which positions the scrollbar correctly. */
+	TAttribute<float> ScrollbarSpacerFillPercentage;
 	/** The main sequencer interface */
 	TWeakPtr<FSequencer> Sequencer;
 	/** The curve editor. */
diff --git a/Engine/Source/Editor/Sequencer/Private/Sequencer.cpp b/Engine/Source/Editor/Sequencer/Private/Sequencer.cpp
--- a/Engine/Source/Editor/Sequencer/Private/Sequencer.cpp
+++ b/Engine/Source/Editor/Sequencer/Private/Sequencer.cpp
@@ -77,7 +77,8 @@ void FSequencer::InitSequencer( const FSequencerInitParams& InitParams, const TA
 			.ViewRange( this, &FSequencer::OnGetViewRange )
 			.ScrubPosition( this, &FSequencer::OnGetScrubPosition )
 			.OnScrubPositionChanged( this, &FSequencer::OnScrubPositionChanged )
-			.OnViewRangeChanged( this, &FSequencer::OnViewRangeChanged, false );
+			.OnViewRangeChanged( this, &FSequencer::OnViewRangeChanged, false )
+			.OnGetAddMenuContent(InitParams.ViewParams.OnGetAddMenuContent);
 
 		// When undo occurs, get a notification so we can make sure our view is up to date
 		GEditor->RegisterForUndo(this);
@@ -370,7 +371,8 @@ void FSequencer::DeleteSection(class UMovieSceneSection* Section)
 	
 	if( bAnythingRemoved )
 	{
-		UpdateRuntimeInstances();
+		// Full refresh required just in case the last section was removed from any track.
+		NotifyMovieSceneDataChanged();
 	}
 }
 
@@ -1255,14 +1257,14 @@ void FSequencer::FilterToSelectedShotSections(bool bZoomToShotBounds)
 	FilterToShotSections(SelectedShotSections, bZoomToShotBounds);
 }
 
-bool FSequencer::CanKeyProperty(const UClass& ObjectClass, const IPropertyHandle& PropertyHandle) const
+bool FSequencer::CanKeyProperty(FCanKeyPropertyParams CanKeyPropertyParams) const
 {
-	return ObjectChangeListener->IsTypeKeyable( ObjectClass, PropertyHandle );
+	return ObjectChangeListener->CanKeyProperty(CanKeyPropertyParams);
 } 
 
-void FSequencer::KeyProperty(const TArray<UObject*>& ObjectsToKey, const class IPropertyHandle& PropertyHandle) 
+void FSequencer::KeyProperty(FKeyPropertyParams KeyPropertyParams) 
 {
-	ObjectChangeListener->KeyProperty( ObjectsToKey, PropertyHandle );
+	ObjectChangeListener->KeyProperty(KeyPropertyParams);
 }
 
 TSharedRef<ISequencerObjectBindingManager> FSequencer::GetObjectBindingManager() const
diff --git a/Engine/Source/Editor/Sequencer/Private/Sequencer.h b/Engine/Source/Editor/Sequencer/Private/Sequencer.h
--- a/Engine/Source/Editor/Sequencer/Private/Sequencer.h
+++ b/Engine/Source/Editor/Sequencer/Private/Sequencer.h
@@ -88,8 +88,8 @@ public:
 	virtual void AddSubMovieScene(UMovieScene* SubMovieScene) override;
 	virtual void FilterToShotSections(const TArray< TWeakObjectPtr<class UMovieSceneSection> >& ShotSections, bool bZoomToShotBounds = true) override;
 	virtual void FilterToSelectedShotSections(bool bZoomToShotBounds = true) override;
-	virtual bool CanKeyProperty(const UClass& ObjectClass, const class IPropertyHandle& PropertyHandle) const override;
-	virtual void KeyProperty(const TArray<UObject*>& ObjectsToKey, const class IPropertyHandle& PropertyHandle) override;
+	virtual bool CanKeyProperty(FCanKeyPropertyParams CanKeyPropertyParams) const override;
+	virtual void KeyProperty(FKeyPropertyParams KeyPropertyParams) override;
 	virtual TSharedRef<ISequencerObjectBindingManager> GetObjectBindingManager() const override;
 	virtual FSequencerSelection* GetSelection() override;
 
diff --git a/Engine/Source/Editor/Sequencer/Private/SequencerCurveOwner.cpp b/Engine/Source/Editor/Sequencer/Private/SequencerCurveOwner.cpp
--- a/Engine/Source/Editor/Sequencer/Private/SequencerCurveOwner.cpp
+++ b/Engine/Source/Editor/Sequencer/Private/SequencerCurveOwner.cpp
@@ -52,8 +52,8 @@ FSequencerCurveOwner::FSequencerCurveOwner( TSharedPtr<FSequencerNodeTree> InSeq
 	{
 		for ( TSharedRef<IKeyArea> KeyArea : KeyAreaNode->GetAllKeyAreas() )
 		{
-			IKeyArea::FCurveInfo* CurveInfo = KeyArea->GetCurveInfo();
-			if ( CurveInfo != nullptr )
+			FRichCurve* RichCurve = KeyArea->GetRichCurve();
+			if ( RichCurve != nullptr )
 			{
 				bool bAddCurve = false;
 				switch ( GetDefault<USequencerSettings>()->GetCurveVisibility() )
@@ -65,15 +65,15 @@ FSequencerCurveOwner::FSequencerCurveOwner( TSharedPtr<FSequencerNodeTree> InSeq
 					bAddCurve = KeyAreaNode->GetSequencer().GetSelection()->IsSelected(KeyAreaNode);
 					break;
 				case ESequencerCurveVisibility::AnimatedCurves:
-					bAddCurve = CurveInfo->Curve->GetNumKeys() > 0;
+					bAddCurve = RichCurve->GetNumKeys() > 0;
 					break;
 				}
 				if ( bAddCurve )
 				{
 					FName CurveName = BuildCurveName(KeyAreaNode);
-					Curves.Add( FRichCurveEditInfo( CurveInfo->Curve, CurveName ) );
-					ConstCurves.Add( FRichCurveEditInfoConst( CurveInfo->Curve, CurveName ) );
-					EditInfoToSectionMap.Add( FRichCurveEditInfo( CurveInfo->Curve, CurveName ), CurveInfo->OwningSection );
+					Curves.Add( FRichCurveEditInfo( RichCurve, CurveName ) );
+					ConstCurves.Add( FRichCurveEditInfoConst( RichCurve, CurveName ) );
+					EditInfoToSectionMap.Add( FRichCurveEditInfo( RichCurve, CurveName ), KeyArea->GetOwningSection() );
 				}
 			}
 		}
diff --git a/Engine/Source/Editor/Sequencer/Private/SequencerDisplayNode.cpp b/Engine/Source/Editor/Sequencer/Private/SequencerDisplayNode.cpp
--- a/Engine/Source/Editor/Sequencer/Private/SequencerDisplayNode.cpp
+++ b/Engine/Source/Editor/Sequencer/Private/SequencerDisplayNode.cpp
@@ -14,6 +14,7 @@
 #include "IKeyArea.h"
 #include "ISequencerObjectBindingManager.h"
 
+#define LOCTEXT_NAMESPACE "SequencerDisplayNode"
 
 class SSequencerObjectTrack : public SLeafWidget
 {
@@ -214,9 +215,14 @@ void FSequencerDisplayNode::AddKeyAreaNode( FName KeyAreaName, const FText& Disp
 	KeyAreaNode->AddKeyArea( KeyArea );
 }
 
-TSharedRef<SWidget> FSequencerDisplayNode::GenerateWidgetForOutliner( TSharedRef<class FSequencer> Sequencer )
+TSharedRef<SWidget> FSequencerDisplayNode::GenerateContainerWidgetForOutliner()
 {
-	return SNew( SAnimationOutlinerView, SharedThis( this ), Sequencer );
+	return SNew( SAnimationOutlinerView, SharedThis( this ), &GetSequencer() );
+}
+
+TSharedRef<SWidget> FSequencerDisplayNode::GenerateEditWidgetForOutliner()
+{
+	return SNew(SSpacer);
 }
 
 TSharedRef<SWidget> FSequencerDisplayNode::GenerateWidgetForSectionArea( const TAttribute< TRange<float> >& ViewRange )
@@ -336,6 +342,20 @@ bool FSectionKeyAreaNode::GetShotFilteredVisibilityToCache() const
 	return true;
 }
 
+TSharedRef<SWidget> FSectionKeyAreaNode::GenerateEditWidgetForOutliner()
+{
+	// @todo - Sequencer - Support multiple sections/key areas?
+	TArray<TSharedRef<IKeyArea>> KeyAreas = GetAllKeyAreas();
+	if (KeyAreas.Num() > 0)
+	{
+		if (KeyAreas[0]->CanCreateKeyEditor())
+		{
+			return KeyAreas[0]->CreateKeyEditor(&GetSequencer());
+		}
+	}
+	return FSequencerDisplayNode::GenerateEditWidgetForOutliner();
+}
+
 void FSectionKeyAreaNode::AddKeyArea( TSharedRef< IKeyArea> KeyArea )
 {
 	KeyAreas.Add( KeyArea );
@@ -389,6 +409,24 @@ void FTrackNode::GetChildKeyAreaNodesRecursively(TArray< TSharedRef<class FSecti
 	}
 }
 
+TSharedRef<SWidget> FTrackNode::GenerateEditWidgetForOutliner()
+{
+	TSharedPtr<FSectionKeyAreaNode> KeyAreaNode = GetTopLevelKeyNode();
+	if (KeyAreaNode.IsValid())
+	{
+		// @todo - Sequencer - Support multiple sections/key areas?
+		TArray<TSharedRef<IKeyArea>> KeyAreas = KeyAreaNode->GetAllKeyAreas();
+		if (KeyAreas.Num() > 0)
+		{
+			if (KeyAreas[0]->CanCreateKeyEditor())
+			{
+				return KeyAreas[0]->CreateKeyEditor(&GetSequencer());
+			}
+		}
+	}
+	return FSequencerDisplayNode::GenerateEditWidgetForOutliner();
+}
+
 int32 FTrackNode::GetMaxRowIndex() const
 {
 	int32 MaxRowIndex = 0;
@@ -454,6 +492,37 @@ bool FObjectBindingNode::GetShotFilteredVisibilityToCache() const
 	return !GetSequencer().IsShotFilteringOn() || GetSequencer().IsObjectUnfilterable(ObjectBinding) || HasVisibleChildren();
 }
 
+TSharedRef<SWidget> FObjectBindingNode::GenerateEditWidgetForOutliner()
+{
+	return
+		SNew(SComboButton)
+		.ButtonStyle(FEditorStyle::Get(), "FlatButton.Light")
+		.OnGetMenuContent(this, &FObjectBindingNode::OnGetAddPropertyTrackMenuContent)
+		.ContentPadding(FMargin(2, 0))
+		.ButtonContent()
+		[
+			SNew(SHorizontalBox)
+
+			+ SHorizontalBox::Slot()
+			.VAlign(VAlign_Center)
+			.AutoWidth()
+			[
+				SNew(STextBlock)
+				.Font(FEditorStyle::Get().GetFontStyle("FontAwesome.8"))
+				.Text(FText::FromString(FString(TEXT("\xf067"))) /*fa-plus*/)
+			]
+
+			+ SHorizontalBox::Slot()
+			.AutoWidth()
+			.Padding(4, 0, 0, 0)
+			[
+				SNew(STextBlock)
+				.Font(FEditorStyle::GetFontStyle("Sequencer.AnimationOutliner.RegularFont"))
+				.Text(LOCTEXT("AddPropertyButton", "Property"))
+			]
+		];
+}
+
 TSharedPtr<SWidget> FObjectBindingNode::OnSummonContextMenu(const FGeometry& MyGeometry, const FPointerEvent& MouseEvent)
 {
 	FSequencer& ParentSequencer = GetSequencer();
@@ -477,6 +546,86 @@ TSharedPtr<SWidget> FObjectBindingNode::OnSummonContextMenu(const FGeometry& MyG
 	return MenuBuilder.MakeWidget();
 }
 
+void GetKeyablePropertyPaths(UClass* Class, UStruct* PropertySource, TArray<UProperty*>& PropertyPath, FSequencer& Sequencer, TArray<TArray<UProperty*>>& KeyablePropertyPaths)
+{
+	for (TFieldIterator<UProperty> PropertyIterator(PropertySource); PropertyIterator; ++PropertyIterator)
+	{
+		UProperty* Property = *PropertyIterator;
+		PropertyPath.Add(Property);
+		if (Sequencer.CanKeyProperty(FCanKeyPropertyParams(Class, PropertyPath)))
+		{
+			KeyablePropertyPaths.Add(PropertyPath);
+		}
+		else
+		{
+			UStructProperty* StructProperty = Cast<UStructProperty>(Property);
+			if (StructProperty != nullptr)
+			{
+				GetKeyablePropertyPaths(Class, StructProperty->Struct, PropertyPath, Sequencer, KeyablePropertyPaths);
+			}
+		}
+		PropertyPath.RemoveAt(PropertyPath.Num() - 1);
+	}
+}
+
+TSharedRef<SWidget> FObjectBindingNode::OnGetAddPropertyTrackMenuContent()
+{
+	TArray<TArray<UProperty*>> KeyablePropertyPaths;
+	TArray<UObject*> BoundObjects;
+	FSequencer& Sequencer = GetSequencer();
+	Sequencer.GetObjectBindingManager()->GetRuntimeObjects(Sequencer.GetRootMovieSceneInstance(), ObjectBinding, BoundObjects);
+	for (UObject* BoundObject : BoundObjects)
+	{
+		TArray<UProperty*> PropertyPath;
+		GetKeyablePropertyPaths(BoundObject->GetClass(), BoundObject->GetClass(), PropertyPath, Sequencer, KeyablePropertyPaths);
+	}
+
+	KeyablePropertyPaths.Sort([](const TArray<UProperty*>& A,const TArray<UProperty*>& B)
+	{
+		for (int32 IndexToCheck = 0; IndexToCheck < A.Num() && IndexToCheck < B.Num(); IndexToCheck++)
+		{
+			int32 CompareResult = A[IndexToCheck]->GetName().Compare(B[IndexToCheck]->GetName());
+			if (CompareResult != 0)
+			{
+				return CompareResult < 0;
+			}
+		}
+		// If a difference in name wasn't found, the shorter path should be first.
+		return A.Num() < B.Num();
+	});
+
+	FMenuBuilder AddTrackMenuBuilder(true, NULL);
+	for (TArray<UProperty*>& KeyablePropertyPath : KeyablePropertyPaths)
+	{
+		FUIAction AddTrackMenuAction(FExecuteAction::CreateSP(this, &FObjectBindingNode::AddTrackForProperty, KeyablePropertyPath));
+		TArray<FString> PropertyNames;
+		for (UProperty* Property : KeyablePropertyPath)
+		{
+			PropertyNames.Add(Property->GetName());
+		}
+		AddTrackMenuBuilder.AddMenuEntry(FText::FromString(FString::Join(PropertyNames, TEXT("."))), FText(), FSlateIcon(), AddTrackMenuAction);
+	}
+	return AddTrackMenuBuilder.MakeWidget();
+}
+
+void FObjectBindingNode::AddTrackForProperty(TArray<UProperty*> PropertyPath)
+{
+	FSequencer& Sequencer = GetSequencer();
+
+	TArray<UObject*> BoundObjects;
+	Sequencer.GetObjectBindingManager()->GetRuntimeObjects(Sequencer.GetRootMovieSceneInstance(), ObjectBinding, BoundObjects);
+
+	TArray<UObject*> KeyableBoundObjects;
+	for (UObject* BoundObject : BoundObjects)
+	{
+		if (Sequencer.CanKeyProperty(FCanKeyPropertyParams(BoundObject->GetClass(), PropertyPath)))
+		{
+			KeyableBoundObjects.Add(BoundObject);
+		}
+	}
+	Sequencer.KeyProperty(FKeyPropertyParams(KeyableBoundObjects, PropertyPath));
+}
+
 float FSectionCategoryNode::GetNodeHeight() const 
 {
 	return SequencerLayoutConstants::CategoryNodeHeight;
@@ -487,3 +636,5 @@ bool FSectionCategoryNode::GetShotFilteredVisibilityToCache() const
 	// this node is only visible if at least one child node is visible
 	return HasVisibleChildren();
 }
+
+#undef LOCTEXT_NAMESPACE
diff --git a/Engine/Source/Editor/Sequencer/Private/SequencerDisplayNode.h b/Engine/Source/Editor/Sequencer/Private/SequencerDisplayNode.h
--- a/Engine/Source/Editor/Sequencer/Private/SequencerDisplayNode.h
+++ b/Engine/Source/Editor/Sequencer/Private/SequencerDisplayNode.h
@@ -80,12 +80,18 @@ public:
 	virtual FText GetDisplayName() const = 0;
 
 	/**
-	 * Generates a widget for display in the animation outliner portion of the track area
+	 * Generates a container widget for tree display in the animation outliner portion of the track area
 	 * 
-	 * @param Sequencer	Sequencer interface to pass to outliner widgets
-	 * @return Generated outliner widget
+	 * @return Generated outliner container widget
 	 */
-	virtual TSharedRef<SWidget> GenerateWidgetForOutliner( TSharedRef<class FSequencer> Sequencer );
+	virtual TSharedRef<SWidget> GenerateContainerWidgetForOutliner();
+
+	/**
+	* Generates a widget editing a row in the animation outliner portion of the track area
+	*
+	* @return Generated outliner edit widget
+	*/
+	virtual TSharedRef<SWidget> GenerateEditWidgetForOutliner();
 
 	/**
 	 * Generates a widget for display in the section area portion of the track area
@@ -238,6 +244,7 @@ public:
 	virtual float GetNodeHeight() const override;
 	virtual FText GetDisplayName() const override { return DisplayName; }
 	virtual bool GetShotFilteredVisibilityToCache() const override;
+	virtual TSharedRef<SWidget> GenerateEditWidgetForOutliner() override;
 
 	/**
 	 * Adds a key area to this node
@@ -295,6 +302,7 @@ public:
 	virtual FText GetDisplayName() const override;
 	virtual bool GetShotFilteredVisibilityToCache() const override;
 	virtual void GetChildKeyAreaNodesRecursively(TArray< TSharedRef<class FSectionKeyAreaNode> >& OutNodes) const override;
+	virtual TSharedRef<SWidget> GenerateEditWidgetForOutliner() override;
 
 	/**
 	 * Adds a section to this node
@@ -366,6 +374,7 @@ public:
 	virtual FText GetDisplayName() const override;
 	virtual float GetNodeHeight() const override;
 	virtual bool GetShotFilteredVisibilityToCache() const override;
+	TSharedRef<SWidget> GenerateEditWidgetForOutliner() override;
 	
 	/** @return The object binding on this node */
 	const FGuid& GetObjectBinding() const { return ObjectBinding; }
@@ -374,6 +383,10 @@ public:
 	virtual TSharedPtr<SWidget> OnSummonContextMenu(const FGeometry& MyGeometry, const FPointerEvent& MouseEvent) override;
 
 private:
+	TSharedRef<SWidget> OnGetAddPropertyTrackMenuContent();
+
+	void AddTrackForProperty(TArray<UProperty*> PropertyPath);
+
 	/** The binding to live objects */
 	FGuid ObjectBinding;
 	/** The default display name of the object which is used if the binding manager doesn't provide one. */
diff --git a/Engine/Source/Editor/Sequencer/Private/SequencerObjectChangeListener.cpp b/Engine/Source/Editor/Sequencer/Private/SequencerObjectChangeListener.cpp
--- a/Engine/Source/Editor/Sequencer/Private/SequencerObjectChangeListener.cpp
+++ b/Engine/Source/Editor/Sequencer/Private/SequencerObjectChangeListener.cpp
@@ -24,75 +24,54 @@ FSequencerObjectChangeListener::~FSequencerObjectChangeListener()
 	GEditor->OnActorMoved().RemoveAll( this );
 }
 
-void FSequencerObjectChangeListener::OnPropertyChanged( const TArray<UObject*>& ChangedObjects, const IPropertyHandle& PropertyHandle, bool bRequireAutoKey ) const
+void FSequencerObjectChangeListener::OnPropertyChanged(const TArray<UObject*>& ChangedObjects, const IPropertyHandle& PropertyHandle, bool bRequireAutoKey) const
 {
-	bool bIsKeyable = false;
-	
-	for( UObject* Object : ChangedObjects )
-	{
-		if( Object )
-		{
-			bIsKeyable |= IsTypeKeyable( *Object->GetClass(), PropertyHandle );
-		}
-	}
+	BroadcastPropertyChanged(FKeyPropertyParams(ChangedObjects, PropertyHandle), bRequireAutoKey);
+}
 
-	if( bIsKeyable )
+void FSequencerObjectChangeListener::BroadcastPropertyChanged( FKeyPropertyParams KeyPropertyParams, bool bRequireAutoKey ) const
+{
+	const UStructProperty* StructProperty = Cast<const UStructProperty>(KeyPropertyParams.PropertyPath.Last());
+	const UStructProperty* ParentStructProperty = nullptr;
+	if (KeyPropertyParams.PropertyPath.Num() > 1)
 	{
-		const UProperty* Property = PropertyHandle.GetProperty();
-		const UStructProperty* StructProperty = Cast<const UStructProperty>(Property);
-		const UStructProperty* ParentStructProperty = nullptr;
-		TSharedPtr<IPropertyHandle> ParentHandle;
-
-		ParentHandle = PropertyHandle.GetParentHandle();
-		if (ParentHandle->IsValidHandle())
-		{
-			ParentStructProperty = Cast<const UStructProperty>(ParentHandle->GetProperty());
-		}
-
-
-		FString PropertyVarName;
-
-		// If not in auto-key allow partial keying of specific inner struct property values;
-		FName InnerStructPropName = !bRequireAutoKey && ParentStructProperty ? Property->GetFName() : NAME_None;
+		ParentStructProperty = Cast<const UStructProperty>(KeyPropertyParams.PropertyPath[KeyPropertyParams.PropertyPath.Num() - 2]);
+	}
 
-		FKeyPropertyParams Params;
+	FPropertyChangedParams Params;
+	Params.bRequireAutoKey = bRequireAutoKey;
+	Params.ObjectsThatChanged = KeyPropertyParams.ObjectsToKey;
+	Params.StructPropertyNameToKey = NAME_None;
 
-		Params.bRequireAutoKey = bRequireAutoKey;
-		Params.ObjectsThatChanged = ChangedObjects;
+	bool bFoundAndBroadcastedDelegate = false;
+	if (ParentStructProperty)
+	{
+		Params.PropertyPath.Append(KeyPropertyParams.PropertyPath.GetData(), KeyPropertyParams.PropertyPath.Num() - 1);
 
-		bool bFoundAndBroadcastedDelegate = false;
-		if (ParentStructProperty)
+		// If the property parent is a struct, see if this property parent can be keyed. (e.g R,G,B,A for a color)
+		FOnAnimatablePropertyChanged Delegate = ClassToPropertyChangedMap.FindRef(ParentStructProperty->Struct->GetFName());
+		if (Delegate.IsBound())
 		{
-
-			Params.PropertyHandle = ParentHandle.Get();
-			Params.PropertyPath = ParentHandle->GeneratePathToProperty();
-
-
-			// If the property parent is a struct, see if this property parent can be keyed. (e.g R,G,B,A for a color)
-			FOnAnimatablePropertyChanged Delegate = ClassToPropertyChangedMap.FindRef(ParentStructProperty->Struct->GetFName());
-			if (Delegate.IsBound())
+			if (!bRequireAutoKey)
 			{
-				Params.InnerStructPropertyName = InnerStructPropName;
-				bFoundAndBroadcastedDelegate = true;
-				Delegate.Broadcast(Params);
+				Params.StructPropertyNameToKey = KeyPropertyParams.PropertyPath.Last()->GetFName();
 			}
+			bFoundAndBroadcastedDelegate = true;
+			Delegate.Broadcast(Params);
 		}
+	}
 
-		if (!bFoundAndBroadcastedDelegate)
+	if (!bFoundAndBroadcastedDelegate)
+	{
+		Params.PropertyPath = KeyPropertyParams.PropertyPath;
+		if (StructProperty)
 		{
-			if (StructProperty)
-			{
-				Params.PropertyHandle = &PropertyHandle;
-				Params.PropertyPath = PropertyHandle.GeneratePathToProperty();
-				ClassToPropertyChangedMap.FindRef(StructProperty->Struct->GetFName()).Broadcast(Params);
-			}
-			else
-			{
-				Params.PropertyHandle = &PropertyHandle;
-				Params.PropertyPath = PropertyHandle.GeneratePathToProperty();
-				// the property in question is not a struct or an inner of the struct. See if it is directly keyable
-				ClassToPropertyChangedMap.FindRef(Property->GetClass()->GetFName()).Broadcast(Params);
-			}
+			ClassToPropertyChangedMap.FindRef(StructProperty->Struct->GetFName()).Broadcast(Params);
+		}
+		else
+		{
+			// the property in question is not a struct or an inner of the struct. See if it is directly keyable
+			ClassToPropertyChangedMap.FindRef(KeyPropertyParams.PropertyPath.Last()->GetClass()->GetFName()).Broadcast(Params);
 		}
 	}
 }
@@ -132,46 +111,42 @@ bool FSequencerObjectChangeListener::FindPropertySetter( const UClass& ObjectCla
 	return bFound;
 }
 
-bool FSequencerObjectChangeListener::IsTypeKeyable(const UClass& ObjectClass, const IPropertyHandle& PropertyHandle) const
+bool FSequencerObjectChangeListener::CanKeyProperty(FCanKeyPropertyParams CanKeyPropertyParams) const
 {
-	const UProperty* Property = PropertyHandle.GetProperty();
-	const UStructProperty* StructProperty = Cast<const UStructProperty>(Property);
+	const UStructProperty* StructProperty = Cast<const UStructProperty>(CanKeyPropertyParams.PropertyPath.Last());
 	const UStructProperty* ParentStructProperty = nullptr;
-
-	const TSharedPtr<IPropertyHandle> ParentHandle = PropertyHandle.GetParentHandle();
-	if(ParentHandle->IsValidHandle())
+	if (CanKeyPropertyParams.PropertyPath.Num() > 1)
 	{
-		ParentStructProperty = Cast<const UStructProperty>(ParentHandle->GetProperty());
+		ParentStructProperty = Cast<const UStructProperty>(CanKeyPropertyParams.PropertyPath[CanKeyPropertyParams.PropertyPath.Num() - 2]);
 	}
-	
 
 	FString PropertyVarName;
 
 	bool bFound = false;
-	if( StructProperty )
+	if ( StructProperty )
 	{
-		bFound = FindPropertySetter( ObjectClass, StructProperty->Struct->GetFName(), StructProperty->GetName() );
+		bFound = FindPropertySetter(*CanKeyPropertyParams.ObjectClass, StructProperty->Struct->GetFName(), StructProperty->GetName() );
 	}
 	
 	if( !bFound && ParentStructProperty )
 	{
-		// If the property parent is a struct, see if this property parent can be keyed. (e.g R,G,B,A for a color)
-		bFound = FindPropertySetter( ObjectClass, ParentStructProperty->Struct->GetFName(), ParentStructProperty->GetName() );
+		// If the property parent is a struct, see if this property parent can be keyed.
+		bFound = FindPropertySetter(*CanKeyPropertyParams.ObjectClass, ParentStructProperty->Struct->GetFName(), ParentStructProperty->GetName());
 	}
 
 	if( !bFound )
 	{
 		// the property in question is not a struct or an inner of the struct. See if it is directly keyable
-		bFound = FindPropertySetter( ObjectClass, Property->GetClass()->GetFName(), Property->GetName() );
+		bFound = FindPropertySetter(*CanKeyPropertyParams.ObjectClass, CanKeyPropertyParams.PropertyPath.Last()->GetClass()->GetFName(), CanKeyPropertyParams.PropertyPath.Last()->GetName());
 	}
 
 	return bFound;
 }
 
-void FSequencerObjectChangeListener::KeyProperty( const TArray<UObject*>& ObjectsToKey, const class IPropertyHandle& PropertyHandle ) const
+void FSequencerObjectChangeListener::KeyProperty(FKeyPropertyParams KeyPropertyParams) const
 {
 	const bool bRequireAutoKey = false;
-	OnPropertyChanged( ObjectsToKey, PropertyHandle, bRequireAutoKey );
+	BroadcastPropertyChanged(KeyPropertyParams, bRequireAutoKey);
 }
 
 void FSequencerObjectChangeListener::OnObjectPreEditChange( UObject* Object, const FEditPropertyChain& PropertyChain )
diff --git a/Engine/Source/Editor/Sequencer/Private/SequencerObjectChangeListener.h b/Engine/Source/Editor/Sequencer/Private/SequencerObjectChangeListener.h
--- a/Engine/Source/Editor/Sequencer/Private/SequencerObjectChangeListener.h
+++ b/Engine/Source/Editor/Sequencer/Private/SequencerObjectChangeListener.h
@@ -2,6 +2,8 @@
 
 #pragma once
 
+//#include "ISequencerObjectChangeListener.h"
+
 class IPropertyHandle;
 
 /**
@@ -16,8 +18,8 @@ public:
 	/** ISequencerObjectChangeListener interface */
 	virtual FOnAnimatablePropertyChanged& GetOnAnimatablePropertyChanged( FName PropertyTypeName ) override;
 	virtual FOnPropagateObjectChanges& GetOnPropagateObjectChanges() override;
-	virtual bool IsTypeKeyable(const UClass& ObjectClass, const IPropertyHandle& PropertyHandle) const override;
-	virtual void KeyProperty( const TArray<UObject*>& ObjectsToKey, const class IPropertyHandle& PropertyHandle ) const override;
+	virtual bool CanKeyProperty(FCanKeyPropertyParams KeyPropertyParams) const override;
+	virtual void KeyProperty(FKeyPropertyParams KeyPropertyParams) const override;
 	virtual void TriggerAllPropertiesChanged(UObject* Object) override;
 
 private:
@@ -48,7 +50,10 @@ private:
 	 *
 	 * @param Object	The object that PostEditChange was called on
 	 */
-	void OnPropertyChanged( const TArray<UObject*>& ChangedObjects, const IPropertyHandle& PropertyValue, bool bRequireAutoKey ) const;
+	void OnPropertyChanged( const TArray<UObject*>& ChangedObjects, const IPropertyHandle& PropertyHandle, bool bRequireAutoKey ) const;
+
+	/** Broadcasts the property change callback to the appropriate listeners. */
+	void BroadcastPropertyChanged(FKeyPropertyParams KeyPropertyParams, bool bRequireAutoKey) const;
 
 	/**
 	 * @return True if an object is valid for listening to property changes 
diff --git a/Engine/Source/Editor/Sequencer/Private/SequencerPrivatePCH.h b/Engine/Source/Editor/Sequencer/Private/SequencerPrivatePCH.h
--- a/Engine/Source/Editor/Sequencer/Private/SequencerPrivatePCH.h
+++ b/Engine/Source/Editor/Sequencer/Private/SequencerPrivatePCH.h
@@ -7,6 +7,7 @@
 
 #include "SBreadcrumbTrail.h"
 
+#include "KeyPropertyParams.h"
 #include "ISequencerModule.h"
 #include "ISequencerObjectChangeListener.h"
 #include "SequencerSelection.h"
diff --git a/Engine/Source/Editor/Sequencer/Public/IKeyArea.h b/Engine/Source/Editor/Sequencer/Public/IKeyArea.h
--- a/Engine/Source/Editor/Sequencer/Public/IKeyArea.h
+++ b/Engine/Source/Editor/Sequencer/Public/IKeyArea.h
@@ -21,24 +21,6 @@ struct FKeyDrawingInfo
  */
 class IKeyArea
 {
-public:
-	/** Defines curve information for this key area. */
-	class FCurveInfo
-	{
-	public:
-		FCurveInfo( FRichCurve* InCurve, UMovieSceneSection* InOwningSection )
-			: Curve(InCurve)
-			, OwningSection(InOwningSection)
-		{
-		}
-
-		/** The curve associated with the key area. */
-		FRichCurve* const Curve;
-
-		/** The movie section which owns curve. */
-		UMovieSceneSection* const OwningSection;
-	};
-
 public:
 	virtual ~IKeyArea() {}
 
@@ -69,9 +51,30 @@ public:
 	 */
 	virtual void DeleteKey(FKeyHandle KeyHandle) = 0;
 
-	/** Gets any curve info associated with this key area.  This can be null but it must be valid in order to be
+	/**
+	 * Adds a key at the specified time if there isn't already a key present.  The value of the added key should
+	 * be the value which would be returned if the animation containing this key area was evaluated at the specified time.
+	 */
+	virtual void AddKeyUnique(float Time) = 0;
+
+	/** Gets the rich curve associated with this key area.  This can be null but it must be valid in order to be
 	  * edited by the curve editor. */
-	virtual FCurveInfo* GetCurveInfo() = 0;
+	virtual FRichCurve* GetRichCurve() = 0;
+
+	/**
+	 * Gets the section which owns this key area.
+	 */
+	virtual UMovieSceneSection* GetOwningSection() = 0;
+
+	/**
+	 * Returns true if this key area can create a key editor widget for the animation outliner.
+	 */
+	virtual bool CanCreateKeyEditor() = 0;
+
+	/**
+	* Creates a key editor for this key area for use in the animation outliner.
+	*/
+	virtual TSharedRef<SWidget> CreateKeyEditor(ISequencer* Sequencer) = 0;
 };
 
 
diff --git a/Engine/Source/Editor/Sequencer/Public/ISequencer.h b/Engine/Source/Editor/Sequencer/Public/ISequencer.h
--- a/Engine/Source/Editor/Sequencer/Public/ISequencer.h
+++ b/Engine/Source/Editor/Sequencer/Public/ISequencer.h
@@ -4,6 +4,7 @@
 
 
 #include "Runtime/MovieSceneCore/Public/IMovieScenePlayer.h"
+#include "KeyPropertyParams.h"
 
 class UMovieScene;
 class FSequencerSelection;
@@ -118,9 +119,9 @@ public:
 	 */
 	virtual class ISequencerObjectChangeListener& GetObjectChangeListener() = 0;
 
-	virtual bool CanKeyProperty(const UClass& ObjectClass, const class IPropertyHandle& PropertyHandle) const = 0;
+	virtual bool CanKeyProperty(FCanKeyPropertyParams CanKeyPropertyParams) const = 0;
 
-	virtual void KeyProperty( const TArray<UObject*>& ObjectsToKey, const class IPropertyHandle& PropertyHandle ) = 0;
+	virtual void KeyProperty(FKeyPropertyParams KeyPropertyParams) = 0;
 
 	virtual void NotifyMovieSceneDataChanged() = 0;
 
diff --git a/Engine/Source/Editor/Sequencer/Public/ISequencerModule.h b/Engine/Source/Editor/Sequencer/Public/ISequencerModule.h
--- a/Engine/Source/Editor/Sequencer/Public/ISequencerModule.h
+++ b/Engine/Source/Editor/Sequencer/Public/ISequencerModule.h
@@ -15,6 +15,8 @@ class ISequencerObjectBindingManager;
  */
 DECLARE_DELEGATE_RetVal_OneParam( TSharedRef<class FMovieSceneTrackEditor>, FOnCreateTrackEditor, TSharedRef<ISequencer> );
 
+DECLARE_DELEGATE_RetVal_OneParam( TSharedRef<SWidget>, FOnGetAddMenuContent, TSharedRef<ISequencer> );
+
 /** Parameters for initializing a Sequencer */
 struct FSequencerViewParams
 {
@@ -24,6 +26,8 @@ struct FSequencerViewParams
 	/** Initial Scrub Position */
 	float InitialScrubPosition;
 
+	FOnGetAddMenuContent OnGetAddMenuContent;
+
 	FSequencerViewParams()
 		: InitalViewRange( 0.0f, 5.0f )
 		, InitialScrubPosition( 0.0f )
diff --git a/Engine/Source/Editor/Sequencer/Public/ISequencerObjectChangeListener.h b/Engine/Source/Editor/Sequencer/Public/ISequencerObjectChangeListener.h
--- a/Engine/Source/Editor/Sequencer/Public/ISequencerObjectChangeListener.h
+++ b/Engine/Source/Editor/Sequencer/Public/ISequencerObjectChangeListener.h
@@ -2,19 +2,11 @@
 
 #pragma once
 
-class IPropertyHandle;
-
-struct FKeyPropertyParams
-{
-	const IPropertyHandle* PropertyHandle;
-	TArray<UObject*> ObjectsThatChanged;
-	FString PropertyPath;
-	FName InnerStructPropertyName;
-	bool bRequireAutoKey;
-};
+#include "KeyPropertyParams.h"
 
+class IPropertyHandle;
 
-DECLARE_MULTICAST_DELEGATE_OneParam( FOnAnimatablePropertyChanged, const FKeyPropertyParams& );
+DECLARE_MULTICAST_DELEGATE_OneParam( FOnAnimatablePropertyChanged, const FPropertyChangedParams& );
 
 DECLARE_MULTICAST_DELEGATE_OneParam ( FOnPropagateObjectChanges, UObject* );
 
@@ -39,7 +31,7 @@ public:
 	 */
 	virtual void TriggerAllPropertiesChanged(UObject* Object) = 0;
 
-	virtual bool IsTypeKeyable(const UClass& ObjectClass, const IPropertyHandle& PropertyHandle) const = 0;
+	virtual bool CanKeyProperty(FCanKeyPropertyParams KeyPropertyParams) const = 0;
 
-	virtual void KeyProperty( const TArray<UObject*>& ObjectsToKey, const class IPropertyHandle& PropertyHandle ) const = 0;
+	virtual void KeyProperty(FKeyPropertyParams KeyPropertyParams) const = 0;
 };
\ No newline at end of file
diff --git a/Engine/Source/Editor/Sequencer/Public/KeyPropertyParams.h b/Engine/Source/Editor/Sequencer/Public/KeyPropertyParams.h
new file mode 100644
--- /dev/null
+++ b/Engine/Source/Editor/Sequencer/Public/KeyPropertyParams.h
@@ -0,0 +1,100 @@
+// Copyright 1998-2015 Epic Games, Inc. All Rights Reserved.
+
+#pragma once
+
+class IPropertyHandle;
+class UClass;
+class UProperty;
+
+/**
+ * Parameters for determining if a property can be keyed.
+ */
+struct SEQUENCER_API FCanKeyPropertyParams
+{
+	FCanKeyPropertyParams();
+
+	/**
+	 * Creates new can key property parameters.
+	 * @param InObjectClass the class of the object which has the property to be keyed.
+	 * @param InPropertyPath an array of UProperty objects which represents a path of properties to get from
+	 *        the root object to the property to be keyed.
+	 */
+	FCanKeyPropertyParams(UClass* InObjectClass, TArray<UProperty*> InPropertyPath);
+
+	/**
+	* Creates new can key property parameters.
+	* @param InObjectClass the class of the object which has the property to be keyed.
+	* @param InPropertyHandle a handle to the property to be keyed.
+	*/
+	FCanKeyPropertyParams(UClass* InObjectClass, const IPropertyHandle& InPropertyHandle);
+
+	/** The class of the object which has the property to be keyed. */
+	const UClass* ObjectClass;
+	/** An array of UProperty objects which represents a path of properties to get from the root object to the
+	property to be keyed. */
+	TArray<UProperty*> PropertyPath;
+};
+
+/**
+ * Parameters for keying a property.
+ */
+struct SEQUENCER_API FKeyPropertyParams
+{
+	FKeyPropertyParams();
+
+	/**
+	* Creates new key property parameters.
+	* @param InObjectsToKey an array of the objects who's property will be keyed.
+	* @param InPropertyPath an array of UProperty objects which represents a path of properties to get from
+	*        the root object to the property to be keyed.
+	*/
+	FKeyPropertyParams(TArray<UObject*> InObjectsToKey, TArray<UProperty*> InPropertyPath);
+
+	/**
+	* Creates new key property parameters.
+	* @param InObjectsToKey an array of the objects who's property will be keyed.
+	* @param InPropertyHandle a handle to the property to be keyed.
+	*/
+	FKeyPropertyParams(TArray<UObject*> InObjectsToKey, const IPropertyHandle& InPropertyHandle);
+
+	/** An array of the objects who's property will be keyed. */
+	TArray<UObject*> ObjectsToKey;
+	/** An array of UProperty objects which represents a path of properties to get from the root object to the
+	property to be keyed. */
+	TArray<UProperty*> PropertyPath;
+};
+
+/**
+ * Parameters for the property changed callback.
+ */
+class SEQUENCER_API FPropertyChangedParams
+{
+public:
+	/**
+	 * Gets the value of the property that changed.
+	 */
+	template<typename ValueType>
+	ValueType* GetPropertyValue() const
+	{
+		void* CurrentObject = ObjectsThatChanged[0];
+		for (int32 i = 0; i < PropertyPath.Num(); i++)
+		{
+			CurrentObject = PropertyPath[i]->ContainerPtrToValuePtr<ValueType>(CurrentObject, 0);
+		}
+		return (ValueType*)CurrentObject;
+	}
+
+	/** Gets the property path as a period seperated string of property names. */
+	FString GetPropertyPathString() const;
+
+	/** An array of the objects that changed. */
+	TArray<UObject*> ObjectsThatChanged;
+	/** An array of UProperty objects which represents a path of properties to get from the root object to the 
+	property that changed. */
+	TArray<UProperty*> PropertyPath;
+	/** Represents the FName of an inner property which should be keyed for a struct property.  If all inner 
+	properties should be keyed, this will be FName::None. */
+	FName StructPropertyNameToKey;
+	/** If autokey is required. */
+	bool bRequireAutoKey;
+};
\ No newline at end of file
diff --git a/Engine/Source/Editor/Sequencer/Public/MovieSceneTrackEditor.h b/Engine/Source/Editor/Sequencer/Public/MovieSceneTrackEditor.h
--- a/Engine/Source/Editor/Sequencer/Public/MovieSceneTrackEditor.h
+++ b/Engine/Source/Editor/Sequencer/Public/MovieSceneTrackEditor.h
@@ -61,7 +61,7 @@ public:
 		}
 	}
 	
-	void AnimatablePropertyChanged( TSubclassOf<UMovieSceneTrack> TrackClass, bool bMustBeAutokeying, FOnKeyProperty OnKeyProperty )
+	void AnimatablePropertyChanged( TSubclassOf<class UMovieSceneTrack> TrackClass, bool bMustBeAutokeying, FOnKeyProperty OnKeyProperty )
 	{
 		check(OnKeyProperty.IsBound());
 
@@ -109,12 +109,17 @@ public:
 
 		if( !Type )
 		{
-			Type = MovieScene->AddTrack( TrackClass, ObjectHandle );
+			Type = AddTrack( MovieScene, ObjectHandle, TrackClass, UniqueTypeName );
 		}
 
 		return Type;
 	}
 
+	virtual UMovieSceneTrack* AddTrack(UMovieScene* FocusedMovieScene, const FGuid& ObjectHandle, TSubclassOf<class UMovieSceneTrack> TrackClass, FName UniqueTypeName)
+	{
+		return FocusedMovieScene->AddTrack(TrackClass, ObjectHandle);
+	}
+
 	UMovieSceneTrack* GetMasterTrack( TSubclassOf<UMovieSceneTrack> TrackClass )
 	{
 		UMovieScene* MovieScene = GetSequencer()->GetFocusedMovieScene();
diff --git a/Engine/Source/Editor/Sequencer/Sequencer.Build.cs b/Engine/Source/Editor/Sequencer/Sequencer.Build.cs
--- a/Engine/Source/Editor/Sequencer/Sequencer.Build.cs
+++ b/Engine/Source/Editor/Sequencer/Sequencer.Build.cs
@@ -40,5 +40,7 @@ public class Sequencer : ModuleRules
 				"WorkspaceMenuStructure"
 			}
 		);
+
+		CircularlyReferencedDependentModules.Add("MovieSceneTools");
 	}
 }
diff --git a/Engine/Source/Editor/UMGEditor/Private/Animation/MarginTrackEditor.cpp b/Engine/Source/Editor/UMGEditor/Private/Animation/MarginTrackEditor.cpp
--- a/Engine/Source/Editor/UMGEditor/Private/Animation/MarginTrackEditor.cpp
+++ b/Engine/Source/Editor/UMGEditor/Private/Animation/MarginTrackEditor.cpp
@@ -69,34 +69,29 @@ TSharedRef<ISequencerSection> FMarginTrackEditor::MakeSectionInterface( UMovieSc
 }
 
 
-void FMarginTrackEditor::OnMarginChanged(  const FKeyPropertyParams& PropertyKeyParams )
+void FMarginTrackEditor::OnMarginChanged(  const FPropertyChangedParams& PropertyChangedParams )
 {
-	FName PropertyName = PropertyKeyParams.PropertyHandle->GetProperty()->GetFName();
-
 	AnimatablePropertyChanged
 	(
 		UMovieSceneMarginTrack::StaticClass(), 
-		PropertyKeyParams.bRequireAutoKey,
-		FOnKeyProperty::CreateRaw(this, &FMarginTrackEditor::OnKeyMargin, &PropertyKeyParams ) 
+		PropertyChangedParams.bRequireAutoKey,
+		FOnKeyProperty::CreateRaw(this, &FMarginTrackEditor::OnKeyMargin, &PropertyChangedParams ) 
 	);
 }
 
 
-void FMarginTrackEditor::OnKeyMargin( float KeyTime, const FKeyPropertyParams* PropertyKeyParams )
+void FMarginTrackEditor::OnKeyMargin( float KeyTime, const FPropertyChangedParams* PropertyChangedParams )
 {
-	TArray<const void*> MarginValues;
-	PropertyKeyParams->PropertyHandle->AccessRawData( MarginValues );
-
-	FName PropertyName = PropertyKeyParams->PropertyHandle->GetProperty()->GetFName();
+	FMargin MarginValue = *PropertyChangedParams->GetPropertyValue<FMargin>();
+	FName PropertyName = PropertyChangedParams->PropertyPath.Last()->GetFName();
 
-	for( int32 ObjectIndex = 0; ObjectIndex < PropertyKeyParams->ObjectsThatChanged.Num(); ++ObjectIndex )
+	for( int32 ObjectIndex = 0; ObjectIndex < PropertyChangedParams->ObjectsThatChanged.Num(); ++ObjectIndex )
 	{
-		UObject* Object = PropertyKeyParams->ObjectsThatChanged[ObjectIndex];
-		FMargin MarginValue = *(const FMargin*)MarginValues[ObjectIndex];
+		UObject* Object = PropertyChangedParams->ObjectsThatChanged[ObjectIndex];
 
 		FMarginKey Key;
-		Key.bAddKeyEvenIfUnchanged = !PropertyKeyParams->bRequireAutoKey;
-		Key.CurveName = PropertyKeyParams->InnerStructPropertyName;
+		Key.bAddKeyEvenIfUnchanged = !PropertyChangedParams->bRequireAutoKey;
+		Key.CurveName = PropertyChangedParams->StructPropertyNameToKey;
 		Key.Value = MarginValue;
 
 		FGuid ObjectHandle = FindOrCreateHandleToObject( Object );
@@ -106,7 +101,7 @@ void FMarginTrackEditor::OnKeyMargin( float KeyTime, const FKeyPropertyParams* P
 			if( ensure( Track ) )
 			{
 				UMovieSceneMarginTrack* MarginTrack = CastChecked<UMovieSceneMarginTrack>(Track);
-				MarginTrack->SetPropertyNameAndPath( PropertyName, PropertyKeyParams->PropertyPath );
+				MarginTrack->SetPropertyNameAndPath( PropertyName, PropertyChangedParams->GetPropertyPathString() );
 				// Find or add a new section at the auto-key time and changing the property same property
 				// AddKeyToSection is not actually a virtual, it's redefined in each class with a different type
 				bool bSuccessfulAdd = MarginTrack->AddKeyToSection( KeyTime, Key );
diff --git a/Engine/Source/Editor/UMGEditor/Private/Animation/MarginTrackEditor.h b/Engine/Source/Editor/UMGEditor/Private/Animation/MarginTrackEditor.h
--- a/Engine/Source/Editor/UMGEditor/Private/Animation/MarginTrackEditor.h
+++ b/Engine/Source/Editor/UMGEditor/Private/Animation/MarginTrackEditor.h
@@ -36,10 +36,10 @@ private:
 	 * @param InObjectsThatChanged	List of objects that changed
 	 * @param PropertyValue			Handle to the property value which changed
 	 */
-	void OnMarginChanged( const struct FKeyPropertyParams& PropertyKeyParams );
+	void OnMarginChanged( const class FPropertyChangedParams& PropertyChangedParams );
 
 	/** Called After OnMarginChanged if we actually can key the margin */
-	void OnKeyMargin( float KeyTime, const struct FKeyPropertyParams* PropertyKeyParams );
+	void OnKeyMargin( float KeyTime, const class FPropertyChangedParams* PropertyChangedParams );
 };
 
 
diff --git a/Engine/Source/Editor/UMGEditor/Private/Animation/Sequencer2DTransformTrackEditor.cpp b/Engine/Source/Editor/UMGEditor/Private/Animation/Sequencer2DTransformTrackEditor.cpp
--- a/Engine/Source/Editor/UMGEditor/Private/Animation/Sequencer2DTransformTrackEditor.cpp
+++ b/Engine/Source/Editor/UMGEditor/Private/Animation/Sequencer2DTransformTrackEditor.cpp
@@ -85,34 +85,30 @@ TSharedRef<ISequencerSection> F2DTransformTrackEditor::MakeSectionInterface( UMo
 }
 
 
-void F2DTransformTrackEditor::OnTransformChanged( const FKeyPropertyParams& PropertyKeyParams )
+void F2DTransformTrackEditor::OnTransformChanged( const FPropertyChangedParams& PropertyChangedParams )
 {
-	FName PropertyName = PropertyKeyParams.PropertyHandle->GetProperty()->GetFName();
-
 	AnimatablePropertyChanged
 	(
 		UMovieScene2DTransformTrack::StaticClass(), 
-		PropertyKeyParams.bRequireAutoKey,
-		FOnKeyProperty::CreateRaw(this, &F2DTransformTrackEditor::OnKeyTransform, &PropertyKeyParams ) 
+		PropertyChangedParams.bRequireAutoKey,
+		FOnKeyProperty::CreateRaw(this, &F2DTransformTrackEditor::OnKeyTransform, &PropertyChangedParams ) 
 	);
 }
 
 
-void F2DTransformTrackEditor::OnKeyTransform( float KeyTime, const FKeyPropertyParams* PropertyKeyParams )
+void F2DTransformTrackEditor::OnKeyTransform( float KeyTime, const FPropertyChangedParams* PropertyChangedParams )
 {
-	TArray<const void*> TransformValues;
-	PropertyKeyParams->PropertyHandle->AccessRawData( TransformValues );
+	FWidgetTransform TransformValue = *PropertyChangedParams->GetPropertyValue<FWidgetTransform>();
 
-	FName PropertyName = PropertyKeyParams->PropertyHandle->GetProperty()->GetFName();
+	FName PropertyName = PropertyChangedParams->PropertyPath.Last()->GetFName();
 
-	for( int32 ObjectIndex = 0; ObjectIndex < PropertyKeyParams->ObjectsThatChanged.Num(); ++ObjectIndex )
+	for( int32 ObjectIndex = 0; ObjectIndex < PropertyChangedParams->ObjectsThatChanged.Num(); ++ObjectIndex )
 	{
-		UObject* Object = PropertyKeyParams->ObjectsThatChanged[ObjectIndex];
-		FWidgetTransform TransformValue = *(const FWidgetTransform*)TransformValues[ObjectIndex];
+		UObject* Object = PropertyChangedParams->ObjectsThatChanged[ObjectIndex];
 
 		F2DTransformKey Key;
-		Key.bAddKeyEvenIfUnchanged = !PropertyKeyParams->bRequireAutoKey;
-		Key.CurveName = PropertyKeyParams->InnerStructPropertyName;
+		Key.bAddKeyEvenIfUnchanged = !PropertyChangedParams->bRequireAutoKey;
+		Key.CurveName = PropertyChangedParams->StructPropertyNameToKey;
 		Key.Value = TransformValue;
 
 		FGuid ObjectHandle = FindOrCreateHandleToObject( Object );
@@ -122,7 +118,7 @@ void F2DTransformTrackEditor::OnKeyTransform( float KeyTime, const FKeyPropertyP
 			if( ensure( Track ) )
 			{
 				UMovieScene2DTransformTrack* TransformTrack = CastChecked<UMovieScene2DTransformTrack>(Track);
-				TransformTrack->SetPropertyNameAndPath( PropertyName, PropertyKeyParams->PropertyPath );
+				TransformTrack->SetPropertyNameAndPath( PropertyName, PropertyChangedParams->GetPropertyPathString() );
 				// Find or add a new section at the auto-key time and changing the property same property
 				// AddKeyToSection is not actually a virtual, it's redefined in each class with a different type
 				bool bSuccessfulAdd = TransformTrack->AddKeyToSection( KeyTime, Key );
diff --git a/Engine/Source/Editor/UMGEditor/Private/Animation/Sequencer2DTransformTrackEditor.h b/Engine/Source/Editor/UMGEditor/Private/Animation/Sequencer2DTransformTrackEditor.h
--- a/Engine/Source/Editor/UMGEditor/Private/Animation/Sequencer2DTransformTrackEditor.h
+++ b/Engine/Source/Editor/UMGEditor/Private/Animation/Sequencer2DTransformTrackEditor.h
@@ -36,10 +36,10 @@ private:
 	 * @param InObjectsThatChanged	List of objects that changed
 	 * @param PropertyValue			Handle to the property value which changed
 	 */
-	void OnTransformChanged( const struct FKeyPropertyParams& PropertyKeyParams );
+	void OnTransformChanged( const class FPropertyChangedParams& PropertyChangedParams );
 
 	/** Called After OnMarginChanged if we actually can key the margin */
-	void OnKeyTransform( float KeyTime, const struct FKeyPropertyParams* PropertyKeyParams );
+	void OnKeyTransform( float KeyTime, const class FPropertyChangedParams* PropertyChangedParams );
 };
 
 
diff --git a/Engine/Source/Editor/UMGEditor/Private/Animation/UMGDetailKeyframeHandler.cpp b/Engine/Source/Editor/UMGEditor/Private/Animation/UMGDetailKeyframeHandler.cpp
--- a/Engine/Source/Editor/UMGEditor/Private/Animation/UMGDetailKeyframeHandler.cpp
+++ b/Engine/Source/Editor/UMGEditor/Private/Animation/UMGDetailKeyframeHandler.cpp
@@ -10,9 +10,9 @@ FUMGDetailKeyframeHandler::FUMGDetailKeyframeHandler(TSharedPtr<FWidgetBlueprint
 	: BlueprintEditor( InBlueprintEditor )
 {}
 
-bool FUMGDetailKeyframeHandler::IsPropertyKeyable(const UClass& InObjectClass, const IPropertyHandle& InPropertyHandle) const
+bool FUMGDetailKeyframeHandler::IsPropertyKeyable(UClass* InObjectClass, const IPropertyHandle& InPropertyHandle) const
 {
-	return BlueprintEditor.Pin()->GetSequencer()->CanKeyProperty( InObjectClass, InPropertyHandle );
+	return BlueprintEditor.Pin()->GetSequencer()->CanKeyProperty(FCanKeyPropertyParams(InObjectClass, InPropertyHandle));
 }
 
 bool FUMGDetailKeyframeHandler::IsPropertyKeyingEnabled() const
@@ -26,5 +26,5 @@ void FUMGDetailKeyframeHandler::OnKeyPropertyClicked(const IPropertyHandle& Keye
 	TArray<UObject*> Objects;
 	KeyedPropertyHandle.GetOuterObjects( Objects );
 
-	BlueprintEditor.Pin()->GetSequencer()->KeyProperty( Objects, KeyedPropertyHandle );
+	BlueprintEditor.Pin()->GetSequencer()->KeyProperty(FKeyPropertyParams(Objects, KeyedPropertyHandle));
 }
diff --git a/Engine/Source/Editor/UMGEditor/Private/Animation/UMGDetailKeyframeHandler.h b/Engine/Source/Editor/UMGEditor/Private/Animation/UMGDetailKeyframeHandler.h
--- a/Engine/Source/Editor/UMGEditor/Private/Animation/UMGDetailKeyframeHandler.h
+++ b/Engine/Source/Editor/UMGEditor/Private/Animation/UMGDetailKeyframeHandler.h
@@ -9,7 +9,7 @@ class FUMGDetailKeyframeHandler : public IDetailKeyframeHandler
 public:
 	FUMGDetailKeyframeHandler( TSharedPtr<class FWidgetBlueprintEditor> InBlueprintEditor );
 
-	virtual bool IsPropertyKeyable(const UClass& InObjectClass, const class IPropertyHandle& PropertyHandle) const override;
+	virtual bool IsPropertyKeyable(UClass* InObjectClass, const class IPropertyHandle& PropertyHandle) const override;
 
 	virtual bool IsPropertyKeyingEnabled() const override;
 
diff --git a/Engine/Source/Editor/UMGEditor/Private/WidgetBlueprintEditor.cpp b/Engine/Source/Editor/UMGEditor/Private/WidgetBlueprintEditor.cpp
--- a/Engine/Source/Editor/UMGEditor/Private/WidgetBlueprintEditor.cpp
+++ b/Engine/Source/Editor/UMGEditor/Private/WidgetBlueprintEditor.cpp
@@ -544,6 +544,7 @@ TSharedPtr<ISequencer>& FWidgetBlueprintEditor::GetSequencer()
 		FSequencerViewParams ViewParams;
 		ViewParams.InitalViewRange = TRange<float>(-0.02f, 3.2f);
 		ViewParams.InitialScrubPosition = 0;
+		ViewParams.OnGetAddMenuContent = FOnGetAddMenuContent::CreateSP(this, &FWidgetBlueprintEditor::OnGetAnimationAddMenuContent);
 
 		SequencerObjectBindingManager = MakeShareable(new FUMGSequencerObjectBindingManager(*this, *UWidgetAnimation::GetNullAnimation()));
 
@@ -699,4 +700,89 @@ EWidgetDesignFlags::Type FWidgetBlueprintEditor::GetCurrentDesignerFlags() const
 	return Flags;
 }
 
+class FObjectAndDisplayName
+{
+public:
+	FObjectAndDisplayName(FText InDisplayName, UObject* InObject)
+	{
+		DisplayName = InDisplayName;
+		Object = InObject;
+	}
+
+	bool operator<(FObjectAndDisplayName const& Other) const
+	{
+		return DisplayName.CompareTo(Other.DisplayName) < 0;
+	}
+
+	FText DisplayName;
+	UObject* Object;
+
+};
+
+void GetBindableObjects(UWidget* RootWidget, TArray<FObjectAndDisplayName>& BindableObjects)
+{
+	TArray<UWidget*> ToTraverse;
+	ToTraverse.Add(RootWidget);
+	while (ToTraverse.Num() > 0)
+	{
+		UWidget* Widget = ToTraverse[0];
+		ToTraverse.RemoveAt(0);
+		BindableObjects.Add(FObjectAndDisplayName(FText::FromString(Widget->GetName()), Widget));
+
+		UUserWidget* UserWidget = Cast<UUserWidget>(Widget);
+		if (UserWidget != nullptr)
+		{
+			TArray<FName> SlotNames;
+			UserWidget->GetSlotNames(SlotNames);
+			for (FName SlotName : SlotNames)
+			{
+				UWidget* Content = UserWidget->GetContentForSlot(SlotName);
+				if (Content != nullptr)
+				{
+					ToTraverse.Add(Content);
+				}
+			}
+		}
+
+		UPanelWidget* PanelWidget = Cast<UPanelWidget>(Widget);
+		if (PanelWidget != nullptr)
+		{
+			for (UPanelSlot* Slot : PanelWidget->GetSlots())
+			{
+				if (Slot->Content != nullptr)
+				{
+					FText SlotDisplayName = FText::Format(LOCTEXT("AddMenuSlotFormat", "{0} ({1} Slot)"), FText::FromString(Slot->Content->GetName()), FText::FromString(PanelWidget->GetName()));
+					BindableObjects.Add(FObjectAndDisplayName(SlotDisplayName, Slot));
+					ToTraverse.Add(Slot->Content);
+				}
+			}
+		}
+	}
+}
+
+TSharedRef<SWidget> FWidgetBlueprintEditor::OnGetAnimationAddMenuContent(TSharedRef<ISequencer> Sequencer)
+{
+	TArray<FObjectAndDisplayName> BindableObjects;
+	GetBindableObjects(GetPreview()->GetRootWidget(), BindableObjects);
+	BindableObjects.Sort();
+
+	FMenuBuilder AddMenuBuilder(true, NULL);
+	for (FObjectAndDisplayName& BindableObject : BindableObjects)
+	{
+		FGuid BoundObjectGuid = SequencerObjectBindingManager->FindGuidForObject(*Sequencer->GetFocusedMovieScene(), *BindableObject.Object);
+		if (BoundObjectGuid.IsValid() == false)
+		{
+			FUIAction AddMenuAction(FExecuteAction::CreateSP(this, &FWidgetBlueprintEditor::AddObjectToAnimation, BindableObject.Object));
+			AddMenuBuilder.AddMenuEntry(BindableObject.DisplayName, FText(), FSlateIcon(), AddMenuAction);
+		}
+	}
+	return AddMenuBuilder.MakeWidget();
+}
+
+void FWidgetBlueprintEditor::AddObjectToAnimation(UObject* ObjectToAnimate)
+{
+	// @todo Sequencer - Make this kind of adding more explicit, this current setup seem a bit brittle.
+	Sequencer->GetHandleToObject(ObjectToAnimate);
+}
+
 #undef LOCTEXT_NAMESPACE
diff --git a/Engine/Source/Editor/UMGEditor/Private/WidgetBlueprintEditor.h b/Engine/Source/Editor/UMGEditor/Private/WidgetBlueprintEditor.h
--- a/Engine/Source/Editor/UMGEditor/Private/WidgetBlueprintEditor.h
+++ b/Engine/Source/Editor/UMGEditor/Private/WidgetBlueprintEditor.h
@@ -175,6 +175,12 @@ private:
 	/** Tick the current preview GUI object */
 	void UpdatePreview(UBlueprint* InBlueprint, bool bInForceFullUpdate);
 
+	/** Populates the sequencer add menu. */
+	TSharedRef<SWidget> OnGetAnimationAddMenuContent(TSharedRef<ISequencer> Sequencer);
+
+	/** Adds the supplied UObject to the current animation. */
+	void AddObjectToAnimation(UObject* ObjectToAnimate);
+
 private:
 	/** The preview scene that owns the preview GUI */
 	FPreviewScene PreviewScene;
diff --git a/Engine/Source/Editor/UMGEditor/UMGEditor.Build.cs b/Engine/Source/Editor/UMGEditor/UMGEditor.Build.cs
--- a/Engine/Source/Editor/UMGEditor/UMGEditor.Build.cs
+++ b/Engine/Source/Editor/UMGEditor/UMGEditor.Build.cs
@@ -48,6 +48,7 @@ public class UMGEditor : ModuleRules
 				"Slate",
 				"SlateCore",
 				"MovieSceneCore",
+				"MovieSceneTools",
 				"Sequencer",
 				"DetailCustomizations",
                 "Settings",
diff --git a/Engine/Source/Runtime/Engine/Classes/Curves/CurveBase.h b/Engine/Source/Runtime/Engine/Classes/Curves/CurveBase.h
--- a/Engine/Source/Runtime/Engine/Classes/Curves/CurveBase.h
+++ b/Engine/Source/Runtime/Engine/Classes/Curves/CurveBase.h
@@ -592,6 +592,8 @@ public:
 	FIntegralKey& GetKey(FKeyHandle KeyHandle);
 	FIntegralKey GetKey(FKeyHandle KeyHandle) const;
 
+	FKeyHandle FindKey(float KeyTime) const;
+
 private:
 	/** The keys, ordered by time */
 	UPROPERTY()
diff --git a/Engine/Source/Runtime/Engine/Private/CurveBase.cpp b/Engine/Source/Runtime/Engine/Private/CurveBase.cpp
--- a/Engine/Source/Runtime/Engine/Private/CurveBase.cpp
+++ b/Engine/Source/Runtime/Engine/Private/CurveBase.cpp
@@ -1131,3 +1131,31 @@ FIntegralKey FIntegralCurve::GetKey(FKeyHandle KeyHandle) const
 	EnsureAllIndicesHaveHandles();
 	return Keys[GetIndex(KeyHandle)];
 }
+
+FKeyHandle FIntegralCurve::FindKey(float KeyTime) const
+{
+	int32 Start = 0;
+	int32 End = Keys.Num() - 1;
+
+	// Binary search since the keys are in sorted order
+	while (Start <= End)
+	{
+		int32 TestPos = Start + (End - Start) / 2;
+
+		float TestKeyTime = Keys[TestPos].Time;
+		if (TestKeyTime == KeyTime)
+		{
+			return GetKeyHandle(TestPos);
+		}
+		else if (TestKeyTime < KeyTime)
+		{
+			Start = TestPos + 1;
+		}
+		else
+		{
+			End = TestPos - 1;
+		}
+	}
+
+	return FKeyHandle();
+}
diff --git a/Engine/Source/Runtime/MovieSceneCoreTypes/Classes/MovieSceneByteTrack.h b/Engine/Source/Runtime/MovieSceneCoreTypes/Classes/MovieSceneByteTrack.h
--- a/Engine/Source/Runtime/MovieSceneCoreTypes/Classes/MovieSceneByteTrack.h
+++ b/Engine/Source/Runtime/MovieSceneCoreTypes/Classes/MovieSceneByteTrack.h
@@ -8,10 +8,10 @@
 #include "MovieSceneByteTrack.generated.h"
 
 /**
- * Handles manipulation of float properties in a movie scene
+ * Handles manipulation of byte properties in a movie scene
  */
-UCLASS( MinimalAPI )
-class UMovieSceneByteTrack : public UMovieScenePropertyTrack
+UCLASS()
+class MOVIESCENECORETYPES_API UMovieSceneByteTrack : public UMovieScenePropertyTrack
 {
 	GENERATED_UCLASS_BODY()
 
@@ -38,4 +38,12 @@ public:
 	 * @return true if anything was evaluated. Note: if false is returned OutByte remains unchanged
 	 */
 	virtual bool Eval( float Position, float LastPostion, uint8& OutByte ) const;
+
+	void SetEnum(UEnum* Enum);
+
+	class UEnum* GetEnum() const;
+
+protected:
+	UPROPERTY()
+	UEnum* Enum;
 };
diff --git a/Engine/Source/Runtime/MovieSceneCoreTypes/Private/MovieSceneByteTrack.cpp b/Engine/Source/Runtime/MovieSceneCoreTypes/Private/MovieSceneByteTrack.cpp
--- a/Engine/Source/Runtime/MovieSceneCoreTypes/Private/MovieSceneByteTrack.cpp
+++ b/Engine/Source/Runtime/MovieSceneCoreTypes/Private/MovieSceneByteTrack.cpp
@@ -49,3 +49,13 @@ bool UMovieSceneByteTrack::Eval( float Position, float LastPostion, uint8& OutBy
 
 	return Section != NULL;
 }
+
+void UMovieSceneByteTrack::SetEnum(UEnum* InEnum)
+{
+	Enum = InEnum;
+}
+
+class UEnum* UMovieSceneByteTrack::GetEnum() const
+{
+	return Enum;
+}
