Sequencer: Camera locked to viewport now gets keyframes when moved.

diff --git a/Engine/Source/Developer/MovieSceneTools/Private/TransformTrackEditor.cpp b/Engine/Source/Developer/MovieSceneTools/Private/TransformTrackEditor.cpp
--- a/Engine/Source/Developer/MovieSceneTools/Private/TransformTrackEditor.cpp
+++ b/Engine/Source/Developer/MovieSceneTools/Private/TransformTrackEditor.cpp
@@ -87,12 +87,19 @@ F3DTransformTrackEditor::F3DTransformTrackEditor( TSharedRef<ISequencer> InSeque
 	// Listen for actor/component movement
 	GEditor->OnBeginObjectMovement().AddRaw( this, &F3DTransformTrackEditor::OnPreTransformChanged );
 	GEditor->OnEndObjectMovement().AddRaw( this, &F3DTransformTrackEditor::OnTransformChanged );
+
+	// Listen for the viewport's viewed through camera starts and stops movement
+	GEditor->OnBeginCameraMovement().AddRaw( this, &F3DTransformTrackEditor::OnPreTransformChanged );
+	GEditor->OnEndCameraMovement().AddRaw( this, &F3DTransformTrackEditor::OnTransformChanged );
 }
 
 F3DTransformTrackEditor::~F3DTransformTrackEditor()
 {
 	GEditor->OnBeginObjectMovement().RemoveAll( this );
 	GEditor->OnEndObjectMovement().RemoveAll( this );
+
+	GEditor->OnBeginCameraMovement().RemoveAll( this );
+	GEditor->OnEndCameraMovement().RemoveAll( this );
 }
 
 TSharedRef<FMovieSceneTrackEditor> F3DTransformTrackEditor::CreateTrackEditor( TSharedRef<ISequencer> InSequencer )
