Sequencer: Fix bool property change so that it converts the bool property value which is a uint32.

diff --git a/Engine/Source/Developer/MovieSceneTools/Private/PropertyTrackEditor.cpp b/Engine/Source/Developer/MovieSceneTools/Private/PropertyTrackEditor.cpp
--- a/Engine/Source/Developer/MovieSceneTools/Private/PropertyTrackEditor.cpp
+++ b/Engine/Source/Developer/MovieSceneTools/Private/PropertyTrackEditor.cpp
@@ -305,7 +305,7 @@ FPropertyTrackEditor::FPropertyTrackEditor( TSharedRef<ISequencer> InSequencer )
 	// Get the object change listener for the sequencer and register a delegates for when properties change that we care about
 	ISequencerObjectChangeListener& ObjectChangeListener = InSequencer->GetObjectChangeListener();
 	ObjectChangeListener.GetOnAnimatablePropertyChanged( NAME_FloatProperty ).AddRaw( this, &FPropertyTrackEditor::OnAnimatedPropertyChanged<float, UMovieSceneFloatTrack> );
-	ObjectChangeListener.GetOnAnimatablePropertyChanged( NAME_BoolProperty ).AddRaw( this, &FPropertyTrackEditor::OnAnimatedPropertyChanged<bool, UMovieSceneBoolTrack> );
+	ObjectChangeListener.GetOnAnimatablePropertyChanged( NAME_BoolProperty ).AddRaw( this, &FPropertyTrackEditor::OnAnimatedBoolPropertyChanged<bool, UMovieSceneBoolTrack> );
 	ObjectChangeListener.GetOnAnimatablePropertyChanged( NAME_ByteProperty ).AddRaw( this, &FPropertyTrackEditor::OnAnimatedPropertyChanged<uint8, UMovieSceneByteTrack> );
 	ObjectChangeListener.GetOnAnimatablePropertyChanged( NAME_Vector ).AddRaw( this, &FPropertyTrackEditor::OnAnimatedVectorPropertyChanged );
 	ObjectChangeListener.GetOnAnimatablePropertyChanged( NAME_Vector4 ).AddRaw( this, &FPropertyTrackEditor::OnAnimatedVectorPropertyChanged );
@@ -431,12 +431,26 @@ template <typename Type, typename TrackType>
 void FPropertyTrackEditor::OnAnimatedPropertyChanged(const FPropertyChangedParams& PropertyChangedParams )
 {
 	// Get the value from the property 
-	Type Value = *PropertyChangedParams.GetPropertyValue<Type>();;
+	Type Value = *PropertyChangedParams.GetPropertyValue<Type>();
 
 	AnimatablePropertyChanged(TrackType::StaticClass(), PropertyChangedParams.bRequireAutoKey,
 		FOnKeyProperty::CreateRaw(this, &FPropertyTrackEditor::OnKeyProperty<Type, TrackType>, PropertyChangedParams, Value) );
 }
 
+template <typename Type, typename TrackType>
+void FPropertyTrackEditor::OnAnimatedBoolPropertyChanged(const FPropertyChangedParams& PropertyChangedParams)
+{
+	// Get the value from the property 
+	const UBoolProperty* BoolProperty = Cast<const UBoolProperty>(PropertyChangedParams.PropertyPath.Last());
+	if (BoolProperty)
+	{
+		bool BoolValue = BoolProperty->GetPropertyValue(BoolProperty->ContainerPtrToValuePtr<void>(PropertyChangedParams.ObjectsThatChanged.Last()));
+
+		AnimatablePropertyChanged(TrackType::StaticClass(), PropertyChangedParams.bRequireAutoKey,
+			FOnKeyProperty::CreateRaw(this, &FPropertyTrackEditor::OnKeyProperty<Type, TrackType>, PropertyChangedParams, BoolValue));
+	}
+}
+
 void FPropertyTrackEditor::OnAnimatedVectorPropertyChanged(const FPropertyChangedParams& PropertyChangedParams)
 {
 	bool bIsVector2D = false, bIsVector = false, bIsVector4 = false;
diff --git a/Engine/Source/Developer/MovieSceneTools/Private/PropertyTrackEditor.h b/Engine/Source/Developer/MovieSceneTools/Private/PropertyTrackEditor.h
--- a/Engine/Source/Developer/MovieSceneTools/Private/PropertyTrackEditor.h
+++ b/Engine/Source/Developer/MovieSceneTools/Private/PropertyTrackEditor.h
@@ -44,6 +44,15 @@ private:
 	template <typename Type, typename TrackType>
 	void OnAnimatedPropertyChanged( const FPropertyChangedParams& KeyPropertyParams );
 
+	/**
+	 * Called by the details panel when an animatable bool property changes
+	 *
+	 * @param InObjectsThatChanged	List of objects that changed
+	 * @param KeyPropertyParams		Parameters for the property change.
+	 */
+	template <typename Type, typename TrackType>
+	void OnAnimatedBoolPropertyChanged( const FPropertyChangedParams& KeyPropertyParams );
+
 	/**
 	 * Called by the details panel when an animatable vector property changes
 	 *
